Live Link Hub: Introduce a new Live Link "Device" framework.

Live Link Devices are designed to enable integrated monitoring and automation of third-party software and hardware in performance capture / stage workflows, similar to Switchboard.

Device implementations derive from ULiveLinkDevice, and additionally from a combination of "capability" UInterfaces derived from ILiveLinkDeviceCapability.

A device capability serves as as an RTTI-queryable protocol defining a set of related APIs, as well as its CDO being able to implement related functionality, such as registering and populating custom columns/cell widgets in the devices UI table.

As UInterfaces, capabilities can be both invoked and implemented via blueprint, and can be invoked via Python. Ultimately we would like to support implementing them via Python as well, but this would depend on speculative future work to expose UInterface functionality to the Python binding layer.

diff --git a/Engine/Plugins/Animation/LiveLinkDevice/LiveLinkDevice.uplugin b/Engine/Plugins/Animation/LiveLinkDevice/LiveLinkDevice.uplugin
new file mode 100644
--- /dev/null
+++ b/Engine/Plugins/Animation/LiveLinkDevice/LiveLinkDevice.uplugin
@@ -0,0 +1,32 @@
+{
+	"FileVersion": 3,
+	"Version": 1,
+	"VersionName": "1.0",
+	"FriendlyName": "Live Link Device Framework",
+	"Description": "Provides interfaces and base classes for implementing Live Link Hub devices",
+	"Category": "Animation",
+	"CreatedBy": "Epic Games, Inc.",
+	"CreatedByURL": "https://epicgames.com",
+	"DocsURL": "",
+	"MarketplaceURL": "",
+	"SupportURL": "",
+	"EnabledByDefault": false,
+	"CanContainContent": false,
+	"IsBetaVersion": false,
+	"IsExperimentalVersion": true,
+	"Hidden": true,
+	"Installed": false,
+	"Plugins": [
+		{
+			"Name": "LiveLinkHub",
+			"Enabled": true
+		}
+	],
+	"Modules": [
+		{
+			"Name": "LiveLinkDevice",
+			"Type": "Editor",
+			"LoadingPhase": "Default"
+		}
+	]
+}
diff --git a/Engine/Plugins/Animation/LiveLinkDevice/Source/LiveLinkDevice/LiveLinkDevice.Build.cs b/Engine/Plugins/Animation/LiveLinkDevice/Source/LiveLinkDevice/LiveLinkDevice.Build.cs
new file mode 100644
--- /dev/null
+++ b/Engine/Plugins/Animation/LiveLinkDevice/Source/LiveLinkDevice/LiveLinkDevice.Build.cs
@@ -0,0 +1,26 @@
+// Copyright Epic Games, Inc. All Rights Reserved.
+
+using UnrealBuildTool;
+
+public class LiveLinkDevice : ModuleRules
+{
+	public LiveLinkDevice(ReadOnlyTargetRules Target) : base(Target)
+	{
+		PrivateDependencyModuleNames.AddRange(
+			new string[] {
+				"Core",
+				"CoreUObject",
+				"Engine",
+				"InputCore",
+				"JsonUtilities",
+				"LiveLinkHub",
+				"Projects",
+				"Slate",
+				"SlateCore",
+				"ToolWidgets",
+			}
+		);
+
+		CppCompileWarningSettings.UnsafeTypeCastWarningLevel = WarningLevel.Error;
+	}
+}
diff --git a/Engine/Plugins/Animation/LiveLinkDevice/Source/LiveLinkDevice/Private/LiveLinkDevice.cpp b/Engine/Plugins/Animation/LiveLinkDevice/Source/LiveLinkDevice/Private/LiveLinkDevice.cpp
new file mode 100644
--- /dev/null
+++ b/Engine/Plugins/Animation/LiveLinkDevice/Source/LiveLinkDevice/Private/LiveLinkDevice.cpp
@@ -0,0 +1,83 @@
+// Copyright Epic Games, Inc. All Rights Reserved.
+
+#include "LiveLinkDevice.h"
+#include "Engine/Engine.h"
+#include "JsonObjectConverter.h"
+#include "LiveLinkDeviceCapability.h"
+#include "LiveLinkDeviceModule.h"
+#include "LiveLinkDeviceSubsystem.h"
+#include "Logging/StructuredLog.h"
+#include "Widgets/SNullWidget.h"
+
+
+TSharedRef<FJsonObject> ULiveLinkDeviceSettings::SaveToJson() const
+{
+	TSharedRef<FJsonObject> JsonObject = MakeShared<FJsonObject>();
+	FJsonObjectConverter::UStructToJsonObject(GetClass(), this, JsonObject);
+	return JsonObject;
+}
+
+
+/** Called to restore device settings from session configuration. */
+bool ULiveLinkDeviceSettings::RestoreFromJson(const TSharedRef<FJsonObject>& InJsonObject)
+{
+	FText FailReason;
+
+	const int64 CheckFlags = 0;
+	const int64 SkipFlags = 0;
+	const bool bStrictMode_false = false;
+	const bool bJsonConverted = FJsonObjectConverter::JsonObjectToUStruct(
+		InJsonObject, GetClass(), this, CheckFlags, SkipFlags, bStrictMode_false, &FailReason
+	);
+
+	if (bJsonConverted)
+	{
+		return true;
+	}
+
+	UE_LOGFMT(LogLiveLinkDevice, Error,
+		"ULiveLinkDeviceSettings::RestoreFromJson: JsonObjectToUStruct failed ({Reason})",
+		FailReason.ToString()
+	);
+
+	return false;
+}
+
+
+TSharedRef<SWidget> ULiveLinkDevice::GenerateWidgetForColumn(const FName InColumnId)
+{
+	// Check if the capability provides a default widget implementation.
+	ULiveLinkDeviceSubsystem* Subsystem = GEngine->GetEngineSubsystem<ULiveLinkDeviceSubsystem>();
+	const TMap<FName, TSubclassOf<ULiveLinkDeviceCapability>>& ColumnIdToCapability = Subsystem->GetTableColumnIdToCapability();
+	const TSubclassOf<ULiveLinkDeviceCapability>* MaybeCapabilityClass = ColumnIdToCapability.Find(InColumnId);
+	if (ensure(MaybeCapabilityClass))
+	{
+		const ULiveLinkDeviceCapability* CapabilityCDO = MaybeCapabilityClass->GetDefaultObject();
+		if (TSharedPtr<SWidget> CapabilityWidget = CapabilityCDO->GenerateWidgetForColumn(InColumnId, this))
+		{
+			return CapabilityWidget.ToSharedRef();
+		}
+	}
+
+	// Neither your device class nor the capability created a widget for this column.
+	ensure(false);
+	return SNullWidget::NullWidget;
+}
+
+
+void ULiveLinkDevice::OnDeviceAdded(const FGuid InDeviceGuid, ULiveLinkDeviceSettings* InSettings)
+{
+	Settings = InSettings;
+}
+
+
+void ULiveLinkDevice::OnDeviceRemoved()
+{
+
+}
+
+
+void ULiveLinkDevice::OnSettingChanged(const FPropertyChangedEvent& InPropertyChangedEvent)
+{
+
+}
diff --git a/Engine/Plugins/Animation/LiveLinkDevice/Source/LiveLinkDevice/Private/LiveLinkDeviceCapability_Connection.cpp b/Engine/Plugins/Animation/LiveLinkDevice/Source/LiveLinkDevice/Private/LiveLinkDeviceCapability_Connection.cpp
new file mode 100644
--- /dev/null
+++ b/Engine/Plugins/Animation/LiveLinkDevice/Source/LiveLinkDevice/Private/LiveLinkDeviceCapability_Connection.cpp
@@ -0,0 +1,29 @@
+// Copyright Epic Games, Inc. All Rights Reserved.
+
+#include "LiveLinkDeviceCapability_Connection.h"
+#include "LiveLinkDevice.h"
+#include "Widgets/Text/STextBlock.h"
+
+
+TSharedPtr<SWidget> ULiveLinkDeviceCapability_Connection::GenerateWidgetForColumn(
+	const FName InColumnId,
+	ULiveLinkDevice* InDevice
+) const
+{
+	return SNew(STextBlock)
+		.Text_Lambda(
+			[WeakDevice = TWeakObjectPtr<ULiveLinkDevice>(InDevice)]
+			()
+			{
+				if (ULiveLinkDevice* Device = WeakDevice.Get())
+				{
+					const FString HardwareId = ILiveLinkDeviceCapability_Connection::Execute_GetHardwareId(Device);
+					return FText::FromString(HardwareId);
+				}
+				else
+				{
+					return FText();
+				}
+			}
+		);
+}
diff --git a/Engine/Plugins/Animation/LiveLinkDevice/Source/LiveLinkDevice/Private/LiveLinkDeviceCapability_Recording.cpp b/Engine/Plugins/Animation/LiveLinkDevice/Source/LiveLinkDevice/Private/LiveLinkDeviceCapability_Recording.cpp
new file mode 100644
--- /dev/null
+++ b/Engine/Plugins/Animation/LiveLinkDevice/Source/LiveLinkDevice/Private/LiveLinkDeviceCapability_Recording.cpp
@@ -0,0 +1,103 @@
+// Copyright Epic Games, Inc. All Rights Reserved.
+
+#include "LiveLinkDeviceCapability_Recording.h"
+#include "Engine/Engine.h"
+#include "ILiveLinkRecordingSessionInfo.h"
+#include "LiveLinkDevice.h"
+#include "LiveLinkDeviceSubsystem.h"
+#include "Widgets/Text/STextBlock.h"
+
+
+void ULiveLinkDeviceCapability_Recording::OnDeviceSubsystemInitialized()
+{
+	IModularFeatures& ModularFeatures = IModularFeatures::Get();
+	if (ModularFeatures.IsModularFeatureAvailable(ILiveLinkRecordingSessionInfo::GetModularFeatureName()))
+	{
+		ILiveLinkRecordingSessionInfo::Get().OnRecordingStarted().AddUObject(
+			this, &ULiveLinkDeviceCapability_Recording::HandleRecordingStarted
+		);
+
+		ILiveLinkRecordingSessionInfo::Get().OnRecordingStopped().AddUObject(
+			this, &ULiveLinkDeviceCapability_Recording::HandleRecordingStopped
+		);
+	}
+}
+
+
+void ULiveLinkDeviceCapability_Recording::OnDeviceSubsystemDeinitializing()
+{
+	IModularFeatures& ModularFeatures = IModularFeatures::Get();
+	if (ModularFeatures.IsModularFeatureAvailable(ILiveLinkRecordingSessionInfo::GetModularFeatureName()))
+	{
+		ILiveLinkRecordingSessionInfo::Get().OnRecordingStarted().RemoveAll(this);
+		ILiveLinkRecordingSessionInfo::Get().OnRecordingStopped().RemoveAll(this);
+	}
+}
+
+
+SHeaderRow::FColumn::FArguments& ULiveLinkDeviceCapability_Recording::GenerateHeaderForColumn(
+	const FName InColumnId,
+	SHeaderRow::FColumn::FArguments& InArgs
+) const
+{
+	if (InColumnId == Column_RecordingStatus)
+	{
+		return InArgs.DefaultLabel(FText::FromString("Rec"))
+			.DefaultTooltip(FText::FromString("Device recording status"))
+			.FillSized(30.0f);
+	}
+
+	return Super::GenerateHeaderForColumn(InColumnId, InArgs);
+}
+
+
+TSharedPtr<SWidget> ULiveLinkDeviceCapability_Recording::GenerateWidgetForColumn(
+	const FName InColumnId,
+	ULiveLinkDevice* InDevice
+) const
+{
+	return SNew(STextBlock)
+		.Text_Lambda(
+			[WeakDevice = TWeakObjectPtr<ULiveLinkDevice>(InDevice)]
+			()
+			{
+				if (ULiveLinkDevice* Device = WeakDevice.Get())
+				{
+					const bool bRecording = ILiveLinkDeviceCapability_Recording::Execute_IsRecording(Device);
+					return bRecording ? FText::FromString(TEXT("R")) : FText::FromString(TEXT("-"));
+				}
+				else
+				{
+					return FText();
+				}
+			}
+		);
+}
+
+
+void ULiveLinkDeviceCapability_Recording::HandleRecordingStarted()
+{
+	ULiveLinkDeviceSubsystem* Subsystem = GEngine->GetEngineSubsystem<ULiveLinkDeviceSubsystem>();
+
+	TArray<ULiveLinkDevice*> RecordingDevices;
+	Subsystem->GetDevicesByCapability(ULiveLinkDeviceCapability_Recording::StaticClass(), RecordingDevices);
+
+	for (ULiveLinkDevice* Device : RecordingDevices)
+	{
+		ILiveLinkDeviceCapability_Recording::Execute_StartRecording(Device);
+	}
+}
+
+
+void ULiveLinkDeviceCapability_Recording::HandleRecordingStopped()
+{
+	ULiveLinkDeviceSubsystem* Subsystem = GEngine->GetEngineSubsystem<ULiveLinkDeviceSubsystem>();
+
+	TArray<ULiveLinkDevice*> RecordingDevices;
+	Subsystem->GetDevicesByCapability(ULiveLinkDeviceCapability_Recording::StaticClass(), RecordingDevices);
+
+	for (ULiveLinkDevice* Device : RecordingDevices)
+	{
+		ILiveLinkDeviceCapability_Recording::Execute_StopRecording(Device);
+	}
+}
diff --git a/Engine/Plugins/Animation/LiveLinkDevice/Source/LiveLinkDevice/Private/LiveLinkDeviceModule.cpp b/Engine/Plugins/Animation/LiveLinkDevice/Source/LiveLinkDevice/Private/LiveLinkDeviceModule.cpp
new file mode 100644
--- /dev/null
+++ b/Engine/Plugins/Animation/LiveLinkDevice/Source/LiveLinkDevice/Private/LiveLinkDeviceModule.cpp
@@ -0,0 +1,226 @@
+// Copyright Epic Games, Inc. All Rights Reserved.
+
+#include "LiveLinkDeviceModule.h"
+#include "DetailsViewArgs.h"
+#include "Engine/Engine.h"
+#include "Features/IModularFeatures.h"
+#include "Framework/Docking/TabManager.h"
+#include "Framework/MultiBox/MultiBoxBuilder.h"
+#include "ILiveLinkHubComponent.h"
+#include "LiveLinkDevice.h"
+#include "LiveLinkDeviceSubsystem.h"
+#include "Logging/StructuredLog.h"
+#include "PropertyEditorModule.h"
+#include "SPositiveActionButton.h"
+#include "Widgets/Docking/SDockTab.h"
+#include "Widgets/DeviceTable.h"
+
+
+DEFINE_LOG_CATEGORY(LogLiveLinkDevice);
+
+
+const FName FLiveLinkDeviceModule::DevicesTabName = FName("LiveLinkDevices");
+const FName FLiveLinkDeviceModule::DeviceDetailsTabName = FName("LiveLinkDeviceDetails");
+
+
+#define LOCTEXT_NAMESPACE "LiveLinkDevice"
+
+
+IMPLEMENT_MODULE(FLiveLinkDeviceModule, LiveLinkDevice);
+
+
+class FLiveLinkHubDeviceComponent : public ILiveLinkHubComponent
+{
+public:
+	virtual void Init(const FLiveLinkHubComponentInitParams& InParams) override
+	{
+		InParams.MainStack->AddTab(FLiveLinkDeviceModule::DevicesTabName, ETabState::OpenedTab);
+		InParams.MainStack->AddTab(FLiveLinkDeviceModule::DeviceDetailsTabName, ETabState::OpenedTab);
+	}
+};
+
+
+//////////////////////////////////////////////////////////////////////////
+
+
+void FLiveLinkDeviceModule::StartupModule()
+{
+	FGlobalTabmanager::Get()
+		->RegisterNomadTabSpawner(DevicesTabName,
+			FOnSpawnTab::CreateRaw(this, &FLiveLinkDeviceModule::OnSpawnDevicesTab))
+		.SetDisplayName(LOCTEXT("DevicesTabDisplayName", "Devices"));
+
+	FGlobalTabmanager::Get()
+		->RegisterNomadTabSpawner(DeviceDetailsTabName,
+			FOnSpawnTab::CreateRaw(this, &FLiveLinkDeviceModule::OnSpawnDeviceDetailsTab))
+		.SetDisplayName(LOCTEXT("DeviceDetailsTabDisplayName", "Device Details"));
+
+	LLHComponent = MakeShared<FLiveLinkHubDeviceComponent>();
+	IModularFeatures::Get().RegisterModularFeature(ILiveLinkHubComponent::ModularFeatureName, LLHComponent.Get());
+}
+
+
+void FLiveLinkDeviceModule::ShutdownModule()
+{
+}
+
+
+void FLiveLinkDeviceModule::DeviceSelectionChanged(ULiveLinkDevice* InSelectedDevice)
+{
+	WeakSelectedDevice = InSelectedDevice;
+
+	if (DetailsView)
+	{
+		DetailsView->SetObject(InSelectedDevice ? InSelectedDevice->GetBaseSettings() : nullptr);
+	}
+}
+
+
+TSharedRef<SWidget> FLiveLinkDeviceModule::OnGenerateAddDeviceMenu()
+{
+	ULiveLinkDeviceSubsystem* DeviceSubsystem = GEngine->GetEngineSubsystem<ULiveLinkDeviceSubsystem>();
+	check(DeviceSubsystem);
+
+	const TSet<TSubclassOf<ULiveLinkDevice>>& DeviceClasses = DeviceSubsystem->GetKnownDeviceClasses();
+
+	const bool bCloseAfterSelection_true = true;
+	FMenuBuilder MenuBuilder(bCloseAfterSelection_true, nullptr);
+	MenuBuilder.BeginSection("DevicesSection", LOCTEXT("DevicesSectionHeading", "Live Link Devices"));
+
+	for (const TSubclassOf<ULiveLinkDevice>& DeviceClass : DeviceClasses)
+	{
+		const ULiveLinkDevice* DeviceCDO = DeviceClass->GetDefaultObject<ULiveLinkDevice>();
+		if (!ensure(DeviceCDO))
+		{
+			continue;
+		}
+
+		TSubclassOf<ULiveLinkDeviceSettings> DeviceSettingsClass = DeviceCDO->StaticGetSettingsClass();
+		MenuBuilder.AddMenuEntry(
+			DeviceClass->GetDisplayNameText(),
+			FText::FromString(DeviceClass->GetPathName()),
+			FSlateIcon(),
+			FUIAction(FExecuteAction::CreateLambda(
+				[DeviceSubsystem, DeviceClass]
+				{
+					UE_LOGFMT(LogLiveLinkDevice, Log, "Creating new Live Link Device of class {DeviceClass}", DeviceClass->GetName());
+					ULiveLinkDevice* NewDevice = NewObject<ULiveLinkDevice>(DeviceSubsystem, DeviceClass);
+					const FGuid DeviceGuid = DeviceSubsystem->AddDevice(NewDevice);
+
+					TSubclassOf<ULiveLinkDeviceSettings> SettingsClass = NewDevice->StaticGetSettingsClass();
+					UE_LOGFMT(LogLiveLinkDevice, Log, "Initializing with {DeviceSettingsClass}", SettingsClass->GetName());
+					ULiveLinkDeviceSettings* NewSettings = NewObject<ULiveLinkDeviceSettings>(NewDevice, SettingsClass);
+					NewDevice->OnDeviceAdded(DeviceGuid, NewSettings);
+				}
+			)),
+			NAME_None,
+			EUserInterfaceActionType::Button
+		);
+	}
+
+	MenuBuilder.EndSection();
+	return MenuBuilder.MakeWidget();
+}
+
+
+TSharedRef<SDockTab> FLiveLinkDeviceModule::OnSpawnDevicesTab(const FSpawnTabArgs& InSpawnTabArgs)
+{
+	TSharedRef<SDockTab> DockTab = SNew(SDockTab)
+		.TabRole(ETabRole::NomadTab);
+
+	DockTab->SetContent(
+		SNew(SVerticalBox)
+		+ SVerticalBox::Slot()
+		.AutoHeight()
+		[
+			SNew(SBorder)
+			.BorderImage(FAppStyle::GetBrush("ToolPanel.GroupBorder"))
+			.Padding(8.0f)
+			[
+				SNew(SHorizontalBox)
+				+ SHorizontalBox::Slot()
+				.AutoWidth()
+				.Padding(0.f, 5.f, 0.f, 5.f)
+				[
+					SNew(SImage)
+					.ColorAndOpacity(FSlateColor::UseForeground())
+					.Image(FSlateIcon("LiveLinkStyle", "LiveLinkHub.Sources.Icon").GetIcon())
+				]
+				+ SHorizontalBox::Slot()
+				.HAlign(HAlign_Left)
+				.Padding(FMargin(4.0, 2.0))
+				[
+					SNew(STextBlock)
+					.Font(FCoreStyle::GetDefaultFontStyle("Regular", 14))
+					.Text(LOCTEXT("DevicesHeaderText", "Devices"))
+				]
+				+ SHorizontalBox::Slot()
+				.AutoWidth()
+				.VAlign(VAlign_Center)
+				[
+					SNew(SPositiveActionButton)
+					.Icon(FAppStyle::Get().GetBrush("Icons.Plus"))
+					.Text(LOCTEXT("AddDevice", "Add Device"))
+					.ToolTipText(LOCTEXT("AddDevice_Tooltip", "Add a new Live Link device"))
+					.OnGetMenuContent_Raw(this, &FLiveLinkDeviceModule::OnGenerateAddDeviceMenu)
+				]
+			]
+		]
+		+ SVerticalBox::Slot()
+		[
+			SNew(SLiveLinkDeviceTable, DockTab)
+		]
+	);
+
+	return DockTab;
+}
+
+
+TSharedRef<SDockTab> FLiveLinkDeviceModule::OnSpawnDeviceDetailsTab(const FSpawnTabArgs& InSpawnTabArgs)
+{
+	FPropertyEditorModule& PropertyEditorModule =
+		FModuleManager::GetModuleChecked<FPropertyEditorModule>("PropertyEditor");
+
+	FDetailsViewArgs DetailsViewArgs;
+	DetailsViewArgs.NameAreaSettings = FDetailsViewArgs::HideNameArea;
+	DetailsViewArgs.bShowPropertyMatrixButton = false;
+	DetailsViewArgs.bShowKeyablePropertiesOption = false;
+	DetailsViewArgs.bShowAnimatedPropertiesOption = false;
+
+	DetailsView = PropertyEditorModule.CreateDetailView(DetailsViewArgs);
+	if (ULiveLinkDevice* SelectedDevice = WeakSelectedDevice.Get())
+	{
+		DetailsView->SetObject(SelectedDevice->GetBaseSettings());
+	}
+
+	DetailsView->OnFinishedChangingProperties().AddLambda(
+		[this]
+		(const FPropertyChangedEvent& InPropertyChangedEvent)
+		{
+			ULiveLinkDevice* Device = nullptr;
+
+			const TArray<TWeakObjectPtr<UObject>>& SelectedObjects = DetailsView->GetSelectedObjects();
+			if (ensure(SelectedObjects.Num() > 0))
+			{
+				if (UObject* SelectedObject = SelectedObjects[0].Get(); ensure(SelectedObject))
+				{
+					Device = SelectedObject->GetTypedOuter<ULiveLinkDevice>();
+				}
+			}
+
+			if (ensure(Device))
+			{
+				Device->OnSettingChanged(InPropertyChangedEvent);
+			}
+		}
+	);
+
+	return SNew(SDockTab)
+		.TabRole(ETabRole::NomadTab)
+		[
+			DetailsView.ToSharedRef()
+		];
+}
+
+
+#undef LOCTEXT_NAMESPACE
diff --git a/Engine/Plugins/Animation/LiveLinkDevice/Source/LiveLinkDevice/Private/LiveLinkDeviceModule.h b/Engine/Plugins/Animation/LiveLinkDevice/Source/LiveLinkDevice/Private/LiveLinkDeviceModule.h
new file mode 100644
--- /dev/null
+++ b/Engine/Plugins/Animation/LiveLinkDevice/Source/LiveLinkDevice/Private/LiveLinkDeviceModule.h
@@ -0,0 +1,56 @@
+// Copyright Epic Games, Inc. All Rights Reserved.
+
+#pragma once
+
+#include "Logging/LogMacros.h"
+#include "Modules/ModuleManager.h"
+#include "Templates/SharedPointer.h"
+#include "UObject/WeakObjectPtr.h"
+
+
+class FLiveLinkHubDeviceComponent;
+class FSpawnTabArgs;
+class IDetailsView;
+class SDockTab;
+class SWidget;
+class ULiveLinkDevice;
+
+
+DECLARE_LOG_CATEGORY_EXTERN(LogLiveLinkDevice, Log, All);
+
+
+class FLiveLinkDeviceModule : public IModuleInterface
+{
+public:
+	static const FName DevicesTabName;
+	static const FName DeviceDetailsTabName;
+
+	/**
+	 * Singleton-like access to this module's interface.  This is just for convenience!
+	 * Beware of calling this during the shutdown phase, though.  Your module might have been unloaded already.
+	 *
+	 * @return Returns singleton instance, loading the module on demand if needed
+	 */
+	static FLiveLinkDeviceModule& Get()
+	{
+		static const FName ModuleName = "LiveLinkDevice";
+		return FModuleManager::LoadModuleChecked<FLiveLinkDeviceModule>(ModuleName);
+	}
+
+	//~ Begin IModuleInterface
+	virtual void StartupModule() override;
+	virtual void ShutdownModule() override;
+	//~ End IModuleInterface
+
+	void DeviceSelectionChanged(ULiveLinkDevice* InSelectedDevice);
+
+private:
+	TSharedRef<SWidget> OnGenerateAddDeviceMenu();
+	TSharedRef<SDockTab> OnSpawnDevicesTab(const FSpawnTabArgs& InSpawnTabArgs);
+	TSharedRef<SDockTab> OnSpawnDeviceDetailsTab(const FSpawnTabArgs& InSpawnTabArgs);
+
+private:
+	TSharedPtr<FLiveLinkHubDeviceComponent> LLHComponent;
+	TSharedPtr<IDetailsView> DetailsView;
+	TWeakObjectPtr<ULiveLinkDevice> WeakSelectedDevice;
+};
diff --git a/Engine/Plugins/Animation/LiveLinkDevice/Source/LiveLinkDevice/Private/LiveLinkDeviceSubsystem.cpp b/Engine/Plugins/Animation/LiveLinkDevice/Source/LiveLinkDevice/Private/LiveLinkDeviceSubsystem.cpp
new file mode 100644
--- /dev/null
+++ b/Engine/Plugins/Animation/LiveLinkDevice/Source/LiveLinkDevice/Private/LiveLinkDeviceSubsystem.cpp
@@ -0,0 +1,199 @@
+// Copyright Epic Games, Inc. All Rights Reserved.
+
+#include "LiveLinkDeviceSubsystem.h"
+#include "Logging/StructuredLog.h"
+#include "LiveLinkDevice.h"
+#include "LiveLinkDeviceCapability.h"
+#include "LiveLinkDeviceModule.h"
+
+
+void ULiveLinkDeviceSubsystem::Initialize(FSubsystemCollectionBase& Collection)
+{
+	Super::Initialize(Collection);
+
+	// Enumerate capability classes
+	{
+		TArray<UClass*> Capabilities;
+		GetDerivedClasses(ULiveLinkDeviceCapability::StaticClass(), Capabilities);
+
+		for (UClass* Capability : Capabilities)
+		{
+			UE_LOGFMT(LogLiveLinkDevice, Verbose, "Discovered device capability {Name}", Capability->GetFName());
+			RegisterCapabilityClass(Capability);
+		}
+	}
+
+	// Enumerate device classes
+	{
+		TArray<UClass*> DeviceClasses;
+		GetDerivedClasses(ULiveLinkDevice::StaticClass(), DeviceClasses);
+
+		for (UClass* DeviceClass : DeviceClasses)
+		{
+			UE_LOGFMT(LogLiveLinkDevice, Verbose, "Discovered device class {Name}", DeviceClass->GetFName());
+			RegisterDeviceClass(DeviceClass);
+		}
+	}
+
+	// Dispatch initialized calls to capability CDOs.
+	for (TSubclassOf<ULiveLinkDeviceCapability>& Capability : KnownCapabilities)
+	{
+		ULiveLinkDeviceCapability* CapabilityCDO = Capability.GetDefaultObject();
+		CapabilityCDO->OnDeviceSubsystemInitialized();
+	}
+}
+
+
+void ULiveLinkDeviceSubsystem::Deinitialize()
+{
+	// Dispatch de-initializing calls to capability CDOs.
+	for (TSubclassOf<ULiveLinkDeviceCapability>& Capability : KnownCapabilities)
+	{
+		ULiveLinkDeviceCapability* CapabilityCDO = Capability.GetDefaultObject();
+		CapabilityCDO->OnDeviceSubsystemDeinitializing();
+	}
+
+	Super::Deinitialize();
+}
+
+
+FGuid ULiveLinkDeviceSubsystem::AddDevice(ULiveLinkDevice* InDevice)
+{
+	check(InDevice);
+
+	const FGuid DeviceGuid = FGuid::NewGuid();
+	UE_LOGFMT(LogLiveLinkDevice, Log,
+		"ULiveLinkDeviceSubsystem: Adding device {DeviceName} with ID {DeviceGuid}.",
+		InDevice->GetFName(), DeviceGuid);
+	Devices.Add(DeviceGuid, InDevice);
+	DevicesByClass.Add(InDevice->GetClass(), InDevice);
+	OnDeviceAddedDelegate.Broadcast(DeviceGuid, InDevice);
+	return DeviceGuid;
+}
+
+
+void ULiveLinkDeviceSubsystem::RemoveDevice(FGuid InDeviceId)
+{
+	if (TObjectPtr<ULiveLinkDevice>* MaybeDevice = Devices.Find(InDeviceId); ensure(MaybeDevice))
+	{
+		UE_LOGFMT(LogLiveLinkDevice, Log,
+			"ULiveLinkDeviceSubsystem: Removing device {DeviceName} with ID {DeviceGuid}.",
+			MaybeDevice->GetFName(), InDeviceId);
+		DevicesByClass.Remove(MaybeDevice->GetClass(), *MaybeDevice);
+		Devices.Remove(InDeviceId);
+		OnDeviceRemovedDelegate.Broadcast(InDeviceId, *MaybeDevice);
+	}
+	else
+	{
+		UE_LOGFMT(LogLiveLinkDevice, Error,
+			"ULiveLinkDeviceSubsystem: Failed to remove device with ID {DeviceGuid}.",
+			InDeviceId);
+	}
+}
+
+
+void ULiveLinkDeviceSubsystem::RemoveDevice(ULiveLinkDevice* InDevice)
+{
+	check(InDevice);
+
+	if (const FGuid* MaybeId = Devices.FindKey(InDevice); ensure(MaybeId))
+	{
+		UE_LOGFMT(LogLiveLinkDevice, Log,
+			"ULiveLinkDeviceSubsystem: Removing device {DeviceName} with ID {DeviceGuid}.",
+			InDevice->GetFName(), *MaybeId);
+		Devices.Remove(*MaybeId);
+		OnDeviceRemovedDelegate.Broadcast(*MaybeId, InDevice);
+	}
+	else
+	{
+		UE_LOGFMT(LogLiveLinkDevice, Error,
+			"ULiveLinkDeviceSubsystem: Failed to remove device {DeviceName}.",
+			InDevice->GetFName());
+	}
+}
+
+
+void ULiveLinkDeviceSubsystem::GetDevicesByClass(TSubclassOf<ULiveLinkDevice> DeviceClass, TArray<ULiveLinkDevice*>& OutDevices) const
+{
+	DevicesByClass.MultiFind(DeviceClass, OutDevices);
+}
+
+
+void ULiveLinkDeviceSubsystem::GetDevicesByCapability(TSubclassOf<ULiveLinkDeviceCapability> Capability, TArray<ULiveLinkDevice*>& OutDevices) const
+{
+	TArray<TSubclassOf<ULiveLinkDevice>> ClassesImplementingCapability;
+	DeviceClassesByCapability.MultiFind(Capability, ClassesImplementingCapability);
+	for (TSubclassOf<ULiveLinkDevice> Class : ClassesImplementingCapability)
+	{
+		GetDevicesByClass(Class, OutDevices);
+	}
+}
+
+
+void ULiveLinkDeviceSubsystem::RegisterCapabilityClass(TSubclassOf<ULiveLinkDeviceCapability> InCapability)
+{
+	if (!ensure(InCapability))
+	{
+		return;
+	}
+
+	if (KnownCapabilities.Find(InCapability))
+	{
+		ensureMsgf(false, TEXT("Tried to register duplicate capability"));
+		return;
+	}
+
+	KnownCapabilities.Add(InCapability);
+
+	// We want to encourage shallow inheritance.
+	if (InCapability->GetSuperClass() != ULiveLinkDeviceCapability::StaticClass())
+	{
+		UE_LOGFMT(
+			LogLiveLinkDevice,
+			Warning,
+			"Device capability {Name} inherits from another capability; this is strongly discouraged.",
+			InCapability->GetFName()
+		);
+	}
+
+	ULiveLinkDeviceCapability* CapabilityCDO = InCapability->GetDefaultObject<ULiveLinkDeviceCapability>();
+
+	for (const TPair<FName, ULiveLinkDeviceCapability::FDeviceTableColumnDesc>& ColumnDesc : CapabilityCDO->GetTableColumns())
+	{
+		TableColumnIdToCapability.Add(ColumnDesc.Key, InCapability);
+	}
+}
+
+
+void ULiveLinkDeviceSubsystem::RegisterDeviceClass(TSubclassOf<ULiveLinkDevice> InDeviceClass)
+{
+	if (!ensure(InDeviceClass))
+	{
+		return;
+	}
+
+	if (KnownDeviceClasses.Find(InDeviceClass))
+	{
+		ensureMsgf(false, TEXT("Tried to register duplicate device class"));
+		return;
+	}
+
+	KnownDeviceClasses.Add(InDeviceClass);
+
+	for (TArray<FImplementedInterface>::TConstIterator It(InDeviceClass->Interfaces); It; ++It)
+	{
+		UClass* InterfaceClass = It->Class;
+		if (!InterfaceClass || !InterfaceClass->IsChildOf(ULiveLinkDeviceCapability::StaticClass()))
+		{
+			continue;
+		}
+
+		TSubclassOf<ULiveLinkDeviceCapability> Capability = InterfaceClass;
+
+		// All capabilities should have been enumerated prior to device classes.
+		ensure(KnownCapabilities.Find(Capability));
+
+		CapabilitiesByDeviceClass.AddUnique(InDeviceClass, Capability);
+		DeviceClassesByCapability.AddUnique(Capability, InDeviceClass);
+	}
+}
diff --git a/Engine/Plugins/Animation/LiveLinkDevice/Source/LiveLinkDevice/Private/Tests/LiveLinkDeviceTests.cpp b/Engine/Plugins/Animation/LiveLinkDevice/Source/LiveLinkDevice/Private/Tests/LiveLinkDeviceTests.cpp
new file mode 100644
--- /dev/null
+++ b/Engine/Plugins/Animation/LiveLinkDevice/Source/LiveLinkDevice/Private/Tests/LiveLinkDeviceTests.cpp
@@ -0,0 +1,43 @@
+// Copyright Epic Games, Inc. All Rights Reserved.
+
+#include "LiveLinkDeviceTests.h"
+#include "Engine/Engine.h"
+#include "Misc/AutomationTest.h"
+#include "UObject/ScriptInterface.h"
+
+#include "LiveLinkDeviceSubsystem.h"
+
+
+IMPLEMENT_SIMPLE_AUTOMATION_TEST(FDeviceBaseTest, "LiveLinkHub.Devices.Base",
+	EAutomationTestFlags::EditorContext | EAutomationTestFlags::EngineFilter)
+
+bool FDeviceBaseTest::RunTest(const FString& Parameters)
+{
+	// Device is actually ULiveLinkDevice_BaseTest, but we only have a base class pointer
+	ULiveLinkDevice* UnknownDevice = NewObject<ULiveLinkDevice_BaseTest>();
+
+	// Can still query and retrieve the implemented capabilities.
+	TestTrueExpr(UnknownDevice->Implements<ULiveLinkDeviceCapability_BaseTest>());
+
+	TScriptInterface<ILiveLinkDeviceCapability_BaseTest> TestCapability =
+		TScriptInterface<ILiveLinkDeviceCapability_BaseTest>(UnknownDevice);
+
+	TestNotNull(
+		TEXT("TScriptInterface<ILiveLinkDeviceCapability_BaseTest> GetInterface()"),
+		TestCapability.GetInterface()
+	);
+
+	const int32 RandomValue = FMath::Rand32();
+	ILiveLinkDeviceCapability_BaseTest::Execute_SetValue(UnknownDevice, RandomValue);
+	TestEqual(
+		TEXT("Randomly generated value set through interface"),
+		RandomValue,
+		ILiveLinkDeviceCapability_BaseTest::Execute_GetValue(UnknownDevice)
+	);
+
+	ULiveLinkDeviceSubsystem* Subsystem = GEngine->GetEngineSubsystem<ULiveLinkDeviceSubsystem>();
+	Subsystem->AddDevice(UnknownDevice);
+	Subsystem->RemoveDevice(UnknownDevice);
+
+	return true;
+}
diff --git a/Engine/Plugins/Animation/LiveLinkDevice/Source/LiveLinkDevice/Private/Tests/LiveLinkDeviceTests.h b/Engine/Plugins/Animation/LiveLinkDevice/Source/LiveLinkDevice/Private/Tests/LiveLinkDeviceTests.h
new file mode 100644
--- /dev/null
+++ b/Engine/Plugins/Animation/LiveLinkDevice/Source/LiveLinkDevice/Private/Tests/LiveLinkDeviceTests.h
@@ -0,0 +1,52 @@
+// Copyright Epic Games, Inc. All Rights Reserved.
+
+#pragma once
+
+#include "LiveLinkDevice.h"
+#include "LiveLinkDeviceCapability.h"
+#include "LiveLinkDeviceTests.generated.h"
+
+
+UINTERFACE()
+class LIVELINKDEVICE_API ULiveLinkDeviceCapability_BaseTest : public ULiveLinkDeviceCapability
+{
+	GENERATED_BODY()
+};
+
+
+class LIVELINKDEVICE_API ILiveLinkDeviceCapability_BaseTest : public ILiveLinkDeviceCapability
+{
+	GENERATED_BODY()
+
+public:
+	UFUNCTION(BlueprintCallable, BlueprintNativeEvent, Category="Live Link Device|Testing")
+	int32 GetValue() const;
+	virtual int32 GetValue_Implementation() const PURE_VIRTUAL(ILiveLinkDeviceCapability_BaseTest::GetValue_Implementation, return 0;);
+
+	UFUNCTION(BlueprintCallable, BlueprintNativeEvent, Category="Live Link Device|Testing")
+	void SetValue(int32 InValue);
+	virtual void SetValue_Implementation(int32 InValue) PURE_VIRTUAL(ILiveLinkDeviceCapability_BaseTest::SetValue_Implementation, );
+};
+
+
+UCLASS()
+class ULiveLinkDevice_BaseTest : public ULiveLinkDevice
+	, public ILiveLinkDeviceCapability_BaseTest
+{
+	GENERATED_BODY()
+
+public:
+	//~ Begin ULiveLinkDevice interface
+	virtual FText GetDisplayName() const override { return FText::FromString("Test Device"); }
+	virtual EDeviceHealth GetDeviceHealth() const override { return EDeviceHealth::Nominal; }
+	virtual FText GetHealthText() const override { return FText::FromString("Test Health"); }
+	//~ End ULiveLinkDevice interface
+
+	//~ Begin ILiveLinkDeviceCapability_BaseTest interface
+	virtual int32 GetValue_Implementation() const override { return TestValue; }
+	virtual void SetValue_Implementation(int32 InValue) override { TestValue = InValue; }
+	//~ End ILiveLinkDeviceCapability_BaseTest interface
+
+protected:
+	int32 TestValue;
+};
diff --git a/Engine/Plugins/Animation/LiveLinkDevice/Source/LiveLinkDevice/Private/Widgets/DeviceTable.cpp b/Engine/Plugins/Animation/LiveLinkDevice/Source/LiveLinkDevice/Private/Widgets/DeviceTable.cpp
new file mode 100644
--- /dev/null
+++ b/Engine/Plugins/Animation/LiveLinkDevice/Source/LiveLinkDevice/Private/Widgets/DeviceTable.cpp
@@ -0,0 +1,235 @@
+// Copyright Epic Games, Inc. All Rights Reserved.
+
+#include "DeviceTable.h"
+#include "Engine/Engine.h"
+#include "Framework/Docking/TabManager.h"
+#include "LiveLinkDevice.h"
+#include "LiveLinkDeviceCapability.h"
+#include "LiveLinkDeviceModule.h"
+#include "LiveLinkDeviceSubsystem.h"
+#include "Widgets/Input/SButton.h"
+
+
+#define LOCTEXT_NAMESPACE "LiveLinkDevice"
+
+
+// Intrinsic/"hard-coded" columns (i.e. not driven by a capability; should any of them be?)
+const FName SLiveLinkDeviceTable::ColumnName_StatusIcon("StatusIcon");
+const FName SLiveLinkDeviceTable::ColumnName_DisplayName("DisplayName");
+const FName SLiveLinkDeviceTable::ColumnName_Remove("Remove");
+
+
+FLiveLinkDeviceRowData::FLiveLinkDeviceRowData(ULiveLinkDevice* InDevice)
+	: WeakDevice(InDevice)
+{
+}
+
+
+void SLiveLinkDeviceTable::Construct(const FArguments& InArgs, TSharedRef<SDockTab> InContainingTab)
+{
+	WeakContainingTab = InContainingTab;
+
+	ULiveLinkDeviceSubsystem* Subsystem = GEngine->GetEngineSubsystem<ULiveLinkDeviceSubsystem>();
+
+	ChildSlot
+	[
+		SAssignNew(DeviceListView, SListView<TSharedPtr<FLiveLinkDeviceRowData>>)
+		.SelectionMode(ESelectionMode::Single)
+		.ListItemsSource(&DeviceRows)
+		.OnSelectionChanged_Lambda(
+			[]
+			(TSharedPtr<FLiveLinkDeviceRowData> InSelectedRow, ESelectInfo::Type InSelectType)
+			{
+				ULiveLinkDevice* Device = nullptr;
+				if (InSelectedRow)
+				{
+					Device = InSelectedRow->WeakDevice.Get();
+				}
+
+				FLiveLinkDeviceModule& Module = FLiveLinkDeviceModule::Get();
+				Module.DeviceSelectionChanged(Device);
+			}
+		)
+		.OnMouseButtonDoubleClick_Static(
+			[]
+			(TSharedPtr<FLiveLinkDeviceRowData> InData)
+			{
+				FGlobalTabmanager::Get()->TryInvokeTab(FLiveLinkDeviceModule::DeviceDetailsTabName);
+			}
+		)
+		.OnGenerateRow_Lambda(
+			[]
+			(TSharedPtr<FLiveLinkDeviceRowData> InData, const TSharedRef<STableViewBase>& InOwnerTableView)
+			{
+				return SNew(SLiveLinkDeviceRowWidget, InData, InOwnerTableView);
+			}
+		)
+		.HeaderRow(
+			SAssignNew(HeaderRow, SHeaderRow)
+		)
+	];
+
+	RegenerateList();
+
+	auto OnDevicesChangedHandler = [this](FGuid InDeviceId, ULiveLinkDevice* InDevice)
+	{
+		RegenerateRows();
+	};
+
+	Subsystem->OnDeviceAdded().AddSPLambda(this, OnDevicesChangedHandler);
+	Subsystem->OnDeviceRemoved().AddSPLambda(this, OnDevicesChangedHandler);
+}
+
+
+SLiveLinkDeviceTable::~SLiveLinkDeviceTable()
+{
+	if (ULiveLinkDeviceSubsystem* Subsystem = GEngine->GetEngineSubsystem<ULiveLinkDeviceSubsystem>())
+	{
+		Subsystem->OnDeviceAdded().RemoveAll(this);
+		Subsystem->OnDeviceRemoved().RemoveAll(this);
+	}
+}
+
+
+void SLiveLinkDeviceTable::RegenerateList()
+{
+	RegenerateColumns();
+	RegenerateRows();
+}
+
+
+void SLiveLinkDeviceTable::RegenerateColumns()
+{
+	HeaderRow->ClearColumns();
+
+	HeaderRow->AddColumn(SHeaderRow::Column(ColumnName_StatusIcon)
+		.DefaultLabel(FText::GetEmpty())
+		.FillSized(30.0f)
+		.HAlignCell(HAlign_Center)
+		.VAlignCell(VAlign_Center));
+
+	HeaderRow->AddColumn(SHeaderRow::Column(ColumnName_DisplayName)
+		.DefaultLabel(LOCTEXT("ColumnLabelDisplayName", "Name"))
+		.FillWidth(1.0f));
+
+	HeaderRow->AddColumn(SHeaderRow::Column(ColumnName_Remove)
+		.DefaultLabel(FText::GetEmpty())
+		.FillSized(30.0f)
+		.HAlignCell(HAlign_Center)
+		.VAlignCell(VAlign_Center));
+
+	// Add columns defined by capabilities.
+	// TODO?: Hide columns if devices implementing that capability aren't present?
+	ULiveLinkDeviceSubsystem* Subsystem = GEngine->GetEngineSubsystem<ULiveLinkDeviceSubsystem>();
+	for (const TSubclassOf<ULiveLinkDeviceCapability>& Capability : Subsystem->GetKnownCapabilities())
+	{
+		const ULiveLinkDeviceCapability* CapabilityCDO = Capability.GetDefaultObject();
+		for (const TPair<FName, ULiveLinkDeviceCapability::FDeviceTableColumnDesc>& Column : CapabilityCDO->GetTableColumns())
+		{
+			const FName ColumnId = Column.Key;
+
+			SHeaderRow::FColumn::FArguments ColumnArgs = SHeaderRow::Column(ColumnId);
+			CapabilityCDO->GenerateHeaderForColumn(ColumnId, ColumnArgs);
+
+			// Keep the "Remove Device" column at the end.
+			const int32 InsertIdx = FMath::Max(0, HeaderRow->GetColumns().Num() - 1);
+
+			HeaderRow->InsertColumn(ColumnArgs, InsertIdx);
+		}
+	}
+}
+
+
+void SLiveLinkDeviceTable::RegenerateRows()
+{
+	ULiveLinkDeviceSubsystem* Subsystem = GEngine->GetEngineSubsystem<ULiveLinkDeviceSubsystem>();
+
+	const TMap<FGuid, TObjectPtr<ULiveLinkDevice>>& DeviceMap = Subsystem->GetDeviceMap();
+	const int32 DeviceCount = DeviceMap.Num();
+	DeviceRows.Empty(DeviceCount);
+	for (const TPair<FGuid, TObjectPtr<ULiveLinkDevice>>& DevicePair : DeviceMap)
+	{
+		DeviceRows.Emplace(MakeShared<FLiveLinkDeviceRowData>(DevicePair.Value));
+	}
+
+	DeviceListView->RequestListRefresh();
+}
+
+
+void SLiveLinkDeviceRowWidget::Construct(
+	const FArguments& InArgs,
+	TSharedPtr<FLiveLinkDeviceRowData> InRowData,
+	const TSharedRef<STableViewBase>& InOwnerTableView
+)
+{
+	RowData = InRowData;
+
+	FSuperRowType::Construct(FSuperRowType::FArguments(), InOwnerTableView);
+}
+
+
+TSharedRef<SWidget> SLiveLinkDeviceRowWidget::GenerateWidgetForColumn(const FName& InColumnId)
+{
+	ULiveLinkDevice* Device = GetDevice();
+
+	if (!ensure(Device))
+	{
+		return SNullWidget::NullWidget;
+	}
+
+	if (InColumnId == SLiveLinkDeviceTable::ColumnName_StatusIcon)
+	{
+		return SNullWidget::NullWidget;
+	}
+
+	if (InColumnId == SLiveLinkDeviceTable::ColumnName_DisplayName)
+	{
+		return SNew(STextBlock)
+			.Text_Lambda([Device]() { return Device->GetDisplayName(); })
+			.ToolTipText_Lambda([Device]() { return Device->GetDisplayName(); });
+	}
+
+	if (InColumnId == SLiveLinkDeviceTable::ColumnName_Remove)
+	{
+		return SNew(SButton)
+			.ButtonStyle(FAppStyle::Get(), "HoverHintOnly")
+			.HAlign(HAlign_Center)
+			.VAlign(VAlign_Center)
+			.ToolTipText(LOCTEXT("RemoveDevice", "Remove selected Live Link device"))
+			.ContentPadding(0.f)
+			.ForegroundColor(FSlateColor::UseForeground())
+			.IsFocusable(false)
+			.OnClicked_Lambda(
+				[Device]
+				{
+					ULiveLinkDeviceSubsystem* Subsystem = GEngine->GetEngineSubsystem<ULiveLinkDeviceSubsystem>();
+					Subsystem->RemoveDevice(Device);
+					return FReply::Handled();
+				})
+			[
+				SNew(SImage)
+				.Image(FAppStyle::GetBrush("Icons.Delete"))
+				.ColorAndOpacity(FSlateColor::UseForeground())
+			];
+	}
+
+	// Below here is logic for capability-driven columns.
+	ULiveLinkDeviceSubsystem* Subsystem = GEngine->GetEngineSubsystem<ULiveLinkDeviceSubsystem>();
+	const TMap<FName, TSubclassOf<ULiveLinkDeviceCapability>>& ColumnIdToCapability = Subsystem->GetTableColumnIdToCapability();
+	if (!ensure(ColumnIdToCapability.Contains(InColumnId)))
+	{
+		return SNullWidget::NullWidget;
+	}
+
+	const TSubclassOf<ULiveLinkDeviceCapability>& CapabilityClass = ColumnIdToCapability.FindChecked(InColumnId);
+	if (!Device->GetClass()->ImplementsInterface(CapabilityClass))
+	{
+		// Device does not implement the capability responsible for populating this column.
+		return SNullWidget::NullWidget;
+	}
+
+	return Device->GenerateWidgetForColumn(InColumnId);
+}
+
+
+#undef LOCTEXT_NAMESPACE
diff --git a/Engine/Plugins/Animation/LiveLinkDevice/Source/LiveLinkDevice/Private/Widgets/DeviceTable.h b/Engine/Plugins/Animation/LiveLinkDevice/Source/LiveLinkDevice/Private/Widgets/DeviceTable.h
new file mode 100644
--- /dev/null
+++ b/Engine/Plugins/Animation/LiveLinkDevice/Source/LiveLinkDevice/Private/Widgets/DeviceTable.h
@@ -0,0 +1,73 @@
+// Copyright Epic Games, Inc. All Rights Reserved.
+
+#pragma once
+
+#include "Widgets/Views/ITableRow.h"
+#include "Widgets/Views/SListView.h" // IWYU pragma: keep
+#include "Widgets/Views/STableRow.h"
+
+
+class SDockTab;
+class STableViewBase;
+class ULiveLinkDevice;
+class ULiveLinkDeviceCapability;
+
+
+struct FLiveLinkDeviceRowData
+{
+	TWeakObjectPtr<ULiveLinkDevice> WeakDevice;
+
+	FLiveLinkDeviceRowData(ULiveLinkDevice* InDevice);
+};
+
+
+class SLiveLinkDeviceTable : public SCompoundWidget
+{
+	SLATE_BEGIN_ARGS(SLiveLinkDeviceTable) {}
+	SLATE_END_ARGS()
+
+public:
+	static const FName ColumnName_StatusIcon;
+	static const FName ColumnName_DisplayName;
+	static const FName ColumnName_Remove;
+
+public:
+	void Construct(const FArguments& InArgs, TSharedRef<SDockTab> InContainingTab);
+	virtual ~SLiveLinkDeviceTable();
+
+private:
+	void RegenerateList();
+	void RegenerateColumns();
+	void RegenerateRows();
+
+private:
+	TWeakPtr<SDockTab> WeakContainingTab;
+
+	TSharedPtr<SHeaderRow> HeaderRow;
+	TArray<TSharedPtr<FLiveLinkDeviceRowData>> DeviceRows;
+	TSharedPtr<SListView<TSharedPtr<FLiveLinkDeviceRowData>>> DeviceListView;
+};
+
+
+class SLiveLinkDeviceRowWidget : public SMultiColumnTableRow<TSharedPtr<FLiveLinkDeviceRowData>>
+{
+	SLATE_BEGIN_ARGS(SLiveLinkDeviceRowWidget) {}
+	SLATE_END_ARGS()
+
+public:
+	void Construct(
+		const FArguments& InArgs,
+		TSharedPtr<FLiveLinkDeviceRowData> InRowData,
+		const TSharedRef<STableViewBase>& InOwnerTableView
+	);
+
+	//~ Begin SMultiColumnTableRow interface
+	TSharedRef<SWidget> GenerateWidgetForColumn(const FName& InColumnId) override;
+	//~ End SMultiColumnTableRow interface
+
+	ULiveLinkDevice* GetDevice() { return RowData ? RowData->WeakDevice.Get() : nullptr; }
+	const ULiveLinkDevice* GetDevice() const { return RowData ? RowData->WeakDevice.Get() : nullptr; }
+
+private:
+	TSharedPtr<FLiveLinkDeviceRowData> RowData;
+};
diff --git a/Engine/Plugins/Animation/LiveLinkDevice/Source/LiveLinkDevice/Public/LiveLinkDevice.h b/Engine/Plugins/Animation/LiveLinkDevice/Source/LiveLinkDevice/Public/LiveLinkDevice.h
new file mode 100644
--- /dev/null
+++ b/Engine/Plugins/Animation/LiveLinkDevice/Source/LiveLinkDevice/Public/LiveLinkDevice.h
@@ -0,0 +1,87 @@
+// Copyright Epic Games, Inc. All Rights Reserved.
+
+#pragma once
+
+#include "Templates/SubclassOf.h"
+#include "LiveLinkDevice.generated.h"
+
+
+class FJsonObject;
+class SWidget;
+
+
+UENUM(BlueprintType)
+enum class EDeviceHealth : uint8
+{
+	Nominal = 0,
+	Info,
+	Warning,
+	Error,
+};
+
+
+UCLASS(Abstract)
+class LIVELINKDEVICE_API ULiveLinkDeviceSettings : public UObject
+{
+	GENERATED_BODY()
+
+	/** Called to persist device settings to session configuration. */
+	virtual TSharedRef<FJsonObject> SaveToJson() const;
+
+	/** Called to restore device settings from session configuration. */
+	virtual bool RestoreFromJson(const TSharedRef<FJsonObject>& InJsonObject);
+};
+
+
+/**
+ * Abstract base class for all Live Link devices.
+ * 
+ * Methods having names beginning with "Static" are expected to return the
+ * same result regardless of what object instance they are invoked upon,
+ * and must be valid to call on the class default object (CDO).
+ */
+UCLASS(Abstract, Blueprintable)
+class LIVELINKDEVICE_API ULiveLinkDevice : public UObject
+{
+	GENERATED_BODY()
+
+public:
+	/** Get the device settings class to be displayed in the details view. */
+	virtual TSubclassOf<ULiveLinkDeviceSettings> StaticGetSettingsClass() const { return ULiveLinkDeviceSettings::StaticClass(); }
+
+public:
+	/** Human-readable display name for this device. */
+	UFUNCTION(BlueprintCallable, Category="Live Link Device")
+	virtual FText GetDisplayName() const PURE_VIRTUAL(ULiveLinkDevice::GetDisplayName, return FText(););
+
+	/** At-a-glance health/"severity" of the device. */
+	UFUNCTION(BlueprintCallable, Category="Live Link Device")
+	virtual EDeviceHealth GetDeviceHealth() const PURE_VIRTUAL(ULiveLinkDevice::GetDeviceHealth, return EDeviceHealth::Error;);
+
+	/** Human-readable explanation for the current device health. */
+	UFUNCTION(BlueprintCallable, Category="Live Link Device")
+	virtual FText GetHealthText() const PURE_VIRTUAL(ULiveLinkDevice::GetHealthText, return FText(););
+
+	/**
+	 * Generate the Slate content, for this device, for the specified column.
+	 * The default implementation will delegate to the capability CDO where appropriate.
+	 */
+	virtual TSharedRef<SWidget> GenerateWidgetForColumn(const FName InColumnId);
+
+public:
+	/** Called when the device has been registered with the device manager, with either initial or restored settings. */
+	virtual void OnDeviceAdded(const FGuid InDeviceGuid, ULiveLinkDeviceSettings* InSettings);
+
+	/** Called when the device is removed from the device manager. */
+	virtual void OnDeviceRemoved();
+
+	/** Called when a setting value has changed via the UI. */
+	virtual void OnSettingChanged(const FPropertyChangedEvent& InPropertyChangedEvent);
+
+	const ULiveLinkDeviceSettings* GetBaseSettings() const { return Settings; }
+	ULiveLinkDeviceSettings* GetBaseSettings() { return Settings; }
+
+protected:
+	UPROPERTY()
+	TObjectPtr<ULiveLinkDeviceSettings> Settings;
+};
diff --git a/Engine/Plugins/Animation/LiveLinkDevice/Source/LiveLinkDevice/Public/LiveLinkDeviceCapability.h b/Engine/Plugins/Animation/LiveLinkDevice/Source/LiveLinkDevice/Public/LiveLinkDeviceCapability.h
new file mode 100644
--- /dev/null
+++ b/Engine/Plugins/Animation/LiveLinkDevice/Source/LiveLinkDevice/Public/LiveLinkDeviceCapability.h
@@ -0,0 +1,96 @@
+// Copyright Epic Games, Inc. All Rights Reserved.
+
+#pragma once
+
+#include "UObject/Interface.h"
+#include "Widgets/Views/SHeaderRow.h"
+#include "LiveLinkDeviceCapability.generated.h"
+
+
+class SWidget;
+class ULiveLinkDevice;
+class ULiveLinkDeviceCapability;
+
+
+/**
+ * Base class for all device capabilities.
+ *
+ * Every IInterface has a corresponding UInterface, which exists only as a
+ * class default object (CDO), and provides a place to manage centralized
+ * state related to the capability.
+ */
+UINTERFACE(Blueprintable)
+class LIVELINKDEVICE_API ULiveLinkDeviceCapability : public UInterface
+{
+	GENERATED_BODY()
+
+public:
+	struct FDeviceTableColumnDesc
+	{
+		// TODO: Default show/hide, sorting hint, ?
+	};
+
+public:
+	ULiveLinkDeviceCapability()
+	{
+		ensure(HasAllFlags(RF_ClassDefaultObject));
+	}
+
+	/** Called at completion of Live Link device subsystem initialization. */
+	virtual void OnDeviceSubsystemInitialized() {}
+
+	/** Called at the beginning of Live Link device subsystem de-initialization. */
+	virtual void OnDeviceSubsystemDeinitializing() {}
+
+	/** Return device table widget columns this capability provides. */
+	const TMap<FName, FDeviceTableColumnDesc>& GetTableColumns() const
+	{
+		return TableColumns;
+	}
+
+	/** Configure the header for the specified column. */
+	virtual SHeaderRow::FColumn::FArguments& GenerateHeaderForColumn(const FName InColumnId, SHeaderRow::FColumn::FArguments& InArgs) const
+	{
+		// This base class implementation should not get called.
+		// If your capability defines table columns, you need to implement this.
+		ensure(false);
+		return InArgs;
+	}
+
+	/** Optional; allows a capability to provide a default widget the device can fall back to. */
+	virtual TSharedPtr<SWidget> GenerateWidgetForColumn(const FName InColumnId, ULiveLinkDevice* InDevice) const
+	{
+		return nullptr;
+	}
+
+protected:
+	/**
+	 * Call this from your derived class to define a new device table column.
+	 *
+	 * @return The full, namespaced column identifier for future reference.
+	 */
+	FName RegisterTableColumn(const FName InColumnShortName)
+	{
+		const FName ColumnId = ExpandColumnShortName(InColumnShortName);
+		TableColumns.Emplace(ColumnId, {});
+		return ColumnId;
+	}
+
+	FName ExpandColumnShortName(FName InShortName) const
+	{
+		FNameBuilder NameBuilder;
+		GetClass()->GetFName().AppendString(NameBuilder);
+		NameBuilder.AppendChar('.');
+		InShortName.AppendString(NameBuilder);
+		return FName(NameBuilder.ToView());
+	}
+
+private:
+	TMap<FName, FDeviceTableColumnDesc> TableColumns;
+};
+
+
+class LIVELINKDEVICE_API ILiveLinkDeviceCapability
+{
+	GENERATED_BODY()
+};
diff --git a/Engine/Plugins/Animation/LiveLinkDevice/Source/LiveLinkDevice/Public/LiveLinkDeviceCapability_Connection.h b/Engine/Plugins/Animation/LiveLinkDevice/Source/LiveLinkDevice/Public/LiveLinkDeviceCapability_Connection.h
new file mode 100644
--- /dev/null
+++ b/Engine/Plugins/Animation/LiveLinkDevice/Source/LiveLinkDevice/Public/LiveLinkDeviceCapability_Connection.h
@@ -0,0 +1,75 @@
+// Copyright Epic Games, Inc. All Rights Reserved.
+
+#pragma once
+
+#include "LiveLinkDeviceCapability.h"
+#include "LiveLinkDeviceCapability_Connection.generated.h"
+
+
+UENUM()
+enum class ELiveLinkDeviceConnectionStatus : uint8
+{
+	Disconnected,
+	Connecting,
+	Connected,
+	Disconnecting,
+};
+
+
+UINTERFACE()
+class LIVELINKDEVICE_API ULiveLinkDeviceCapability_Connection : public ULiveLinkDeviceCapability
+{
+	GENERATED_BODY()
+
+public:
+	const FName Column_HardwareId;
+
+	ULiveLinkDeviceCapability_Connection()
+		: Column_HardwareId(RegisterTableColumn("HardwareId"))
+	{
+	}
+
+	virtual SHeaderRow::FColumn::FArguments& GenerateHeaderForColumn(const FName InColumnId, SHeaderRow::FColumn::FArguments& InArgs) const override
+	{
+		if (InColumnId == Column_HardwareId)
+		{
+			return InArgs
+				.DefaultLabel(FText::FromString("Hardware ID"))
+				.DefaultTooltip(FText::FromString("Device hardware identifier"))
+				.FillWidth(0.5f);
+		}
+
+		return Super::GenerateHeaderForColumn(InColumnId, InArgs);
+	}
+
+	virtual TSharedPtr<SWidget> GenerateWidgetForColumn(const FName InColumnId, ULiveLinkDevice* InDevice) const override;
+};
+
+
+class LIVELINKDEVICE_API ILiveLinkDeviceCapability_Connection : public ILiveLinkDeviceCapability
+{
+	GENERATED_BODY()
+
+public:
+	UFUNCTION(BlueprintCallable, BlueprintNativeEvent, Category="Live Link Device|Connection")
+	ELiveLinkDeviceConnectionStatus GetConnectionStatus() const;
+	virtual ELiveLinkDeviceConnectionStatus GetConnectionStatus_Implementation() const PURE_VIRTUAL(ILiveLinkDeviceCapability_Connection::GetConnectionStatus_Implementation, return ELiveLinkDeviceConnectionStatus::Disconnected;);
+
+	/** Retrieve hardware identifier (serial number, network endpoint, etc). */
+	UFUNCTION(BlueprintCallable, BlueprintNativeEvent, Category="Live Link Device|Connection")
+	FString GetHardwareId() const;
+	virtual FString GetHardwareId_Implementation() const PURE_VIRTUAL(ILiveLinkDeviceCapability_Connection::GetHardwareId_Implementation, return FString(););
+
+	/** Set hardware identifier (serial number, network endpoint, etc). */
+	UFUNCTION(BlueprintCallable, BlueprintNativeEvent, Category="Live Link Device|Connection")
+	bool SetHardwareId(const FString& HardwareID);
+	virtual bool SetHardwareId_Implementation(const FString& HardwareID) PURE_VIRTUAL(ILiveLinkDeviceCapability_Connection::SetHardwareId_Implementation, return false;);
+
+	UFUNCTION(BlueprintCallable, BlueprintNativeEvent, Category="Live Link Device|Connection")
+	bool Connect();
+	virtual bool Connect_Implementation() PURE_VIRTUAL(ILiveLinkDeviceCapability_Connection::Connect_Implementation, return false;);
+
+	UFUNCTION(BlueprintCallable, BlueprintNativeEvent, Category="Live Link Device|Connection")
+	bool Disconnect();
+	virtual bool Disconnect_Implementation() PURE_VIRTUAL(ILiveLinkDeviceCapability_Connection::Disconnect_Implementation, return false;);
+};
diff --git a/Engine/Plugins/Animation/LiveLinkDevice/Source/LiveLinkDevice/Public/LiveLinkDeviceCapability_Recording.h b/Engine/Plugins/Animation/LiveLinkDevice/Source/LiveLinkDevice/Public/LiveLinkDeviceCapability_Recording.h
new file mode 100644
--- /dev/null
+++ b/Engine/Plugins/Animation/LiveLinkDevice/Source/LiveLinkDevice/Public/LiveLinkDeviceCapability_Recording.h
@@ -0,0 +1,50 @@
+// Copyright Epic Games, Inc. All Rights Reserved.
+
+#pragma once
+
+#include "LiveLinkDeviceCapability.h"
+#include "LiveLinkDeviceCapability_Recording.generated.h"
+
+
+UINTERFACE()
+class LIVELINKDEVICE_API ULiveLinkDeviceCapability_Recording : public ULiveLinkDeviceCapability
+{
+	GENERATED_BODY()
+
+public:
+	const FName Column_RecordingStatus;
+
+	ULiveLinkDeviceCapability_Recording()
+		: Column_RecordingStatus(RegisterTableColumn("RecordingStatus"))
+	{
+	}
+
+	virtual void OnDeviceSubsystemInitialized() override;
+	virtual void OnDeviceSubsystemDeinitializing() override;
+
+	virtual SHeaderRow::FColumn::FArguments& GenerateHeaderForColumn(const FName InColumnId, SHeaderRow::FColumn::FArguments& InArgs) const override;
+	virtual TSharedPtr<SWidget> GenerateWidgetForColumn(const FName InColumnId, ULiveLinkDevice* InDevice) const override;
+
+protected:
+	void HandleRecordingStarted();
+	void HandleRecordingStopped();
+};
+
+
+class LIVELINKDEVICE_API ILiveLinkDeviceCapability_Recording : public ILiveLinkDeviceCapability
+{
+	GENERATED_BODY()
+
+public:
+	UFUNCTION(BlueprintCallable, BlueprintNativeEvent, Category="Live Link Device|Recording")
+	bool StartRecording();
+	virtual bool StartRecording_Implementation() PURE_VIRTUAL(ILiveLinkDeviceCapability_Recording::StartRecording_Implementation, return false;);
+
+	UFUNCTION(BlueprintCallable, BlueprintNativeEvent, Category="Live Link Device|Recording")
+	bool StopRecording();
+	virtual bool StopRecording_Implementation() PURE_VIRTUAL(ILiveLinkDeviceCapability_Recording::StopRecording_Implementation, return false;);
+
+	UFUNCTION(BlueprintCallable, BlueprintNativeEvent, Category="Live Link Device|Recording")
+	bool IsRecording() const;
+	virtual bool IsRecording_Implementation() const PURE_VIRTUAL(ILiveLinkDeviceCapability_Recording::IsRecording_Implementation, return false;);
+};
diff --git a/Engine/Plugins/Animation/LiveLinkDevice/Source/LiveLinkDevice/Public/LiveLinkDeviceSubsystem.h b/Engine/Plugins/Animation/LiveLinkDevice/Source/LiveLinkDevice/Public/LiveLinkDeviceSubsystem.h
new file mode 100644
--- /dev/null
+++ b/Engine/Plugins/Animation/LiveLinkDevice/Source/LiveLinkDevice/Public/LiveLinkDeviceSubsystem.h
@@ -0,0 +1,93 @@
+// Copyright Epic Games, Inc. All Rights Reserved.
+
+#pragma once
+
+#include "Subsystems/EngineSubsystem.h"
+#include "Templates/SubclassOf.h"
+#include "LiveLinkDeviceSubsystem.generated.h"
+
+
+class ULiveLinkDevice;
+class ULiveLinkDeviceCapability;
+
+
+/**
+ * Device repository with lifecycle notifications.
+ * Facilitates cached lookups related to device and capability classes.
+ */
+UCLASS(BlueprintType)
+class ULiveLinkDeviceSubsystem : public UEngineSubsystem
+{
+	GENERATED_BODY()
+
+public:
+	DECLARE_MULTICAST_DELEGATE_TwoParams(FOnLiveLinkDeviceChangedDelegate, FGuid, ULiveLinkDevice*);
+
+	//~ Begin USubsystem interface
+	virtual void Initialize(FSubsystemCollectionBase& Collection) override;
+	virtual void Deinitialize() override;
+	//~ End USubsystem interface
+
+	/**
+	 * Register a newly created \c ULiveLinkDevice.
+	 * @return A new GUID serving as a handle to this device.
+	 */
+	FGuid AddDevice(ULiveLinkDevice* InDevice);
+
+	/** Remove a previously registered device. */
+	void RemoveDevice(ULiveLinkDevice* InDevice);
+
+	/** Remove a previously registered device by its GUID. */
+	void RemoveDevice(FGuid InDeviceId);
+
+	FOnLiveLinkDeviceChangedDelegate& OnDeviceAdded() { return OnDeviceAddedDelegate; }
+	FOnLiveLinkDeviceChangedDelegate& OnDeviceRemoved() { return OnDeviceRemovedDelegate; }
+
+	/** @return All added ULiveLinkDevice instances and their corresponding GUIDs. */
+	const TMap<FGuid, TObjectPtr<ULiveLinkDevice>>& GetDeviceMap() const { return Devices; }
+
+	/** @return All registered device capabilities. */
+	const TSet<TSubclassOf<ULiveLinkDeviceCapability>>& GetKnownCapabilities() const { return KnownCapabilities; }
+
+	/** @return All registered device classes. */
+	const TSet<TSubclassOf<ULiveLinkDevice>>& GetKnownDeviceClasses() const { return KnownDeviceClasses; }
+
+	/** @return A map from device classes to all capabilities implemented by that device class. */
+	const TMultiMap<TSubclassOf<ULiveLinkDevice>, TSubclassOf<ULiveLinkDeviceCapability>>& GetCapabilitiesByDeviceClass() { return CapabilitiesByDeviceClass; }
+
+	/** @return A map from capability classes to all device classes implementing that capability. */
+	const TMultiMap<TSubclassOf<ULiveLinkDeviceCapability>, TSubclassOf<ULiveLinkDevice>>& GetDeviceClassesByCapability() { return DeviceClassesByCapability; }
+
+	/** @return A map from device table column to capability class. */
+	const TMap<FName, TSubclassOf<ULiveLinkDeviceCapability>>& GetTableColumnIdToCapability() { return TableColumnIdToCapability; }
+
+	UFUNCTION(BlueprintPure, Category="Live Link|Devices", meta=(DeterminesOutputType="DeviceClass"))
+	void GetDevicesByClass(TSubclassOf<ULiveLinkDevice> DeviceClass, TArray<ULiveLinkDevice*>& OutDevices) const;
+
+	UFUNCTION(BlueprintPure, Category="Live Link|Devices")
+	void GetDevicesByCapability(TSubclassOf<ULiveLinkDeviceCapability> Capability, TArray<ULiveLinkDevice*>& OutDevices) const;
+
+private:
+	void RegisterCapabilityClass(TSubclassOf<ULiveLinkDeviceCapability> InCapability);
+	void RegisterDeviceClass(TSubclassOf<ULiveLinkDevice> InDeviceClass);
+
+private:
+	/** Devices by ID (strong references). */
+	UPROPERTY()
+	TMap<FGuid, TObjectPtr<ULiveLinkDevice>> Devices;
+
+	/** Devices by class. */
+	TMultiMap<TSubclassOf<ULiveLinkDevice>, ULiveLinkDevice*> DevicesByClass;
+
+	FOnLiveLinkDeviceChangedDelegate OnDeviceAddedDelegate;
+	FOnLiveLinkDeviceChangedDelegate OnDeviceRemovedDelegate;
+
+	//////////////////////////////////////////////////////////////////////////
+	// UClass/UInterface metadata
+	TSet<TSubclassOf<ULiveLinkDeviceCapability>> KnownCapabilities;
+	TSet<TSubclassOf<ULiveLinkDevice>> KnownDeviceClasses;
+	TMultiMap<TSubclassOf<ULiveLinkDevice>, TSubclassOf<ULiveLinkDeviceCapability>> CapabilitiesByDeviceClass;
+	TMultiMap<TSubclassOf<ULiveLinkDeviceCapability>, TSubclassOf<ULiveLinkDevice>> DeviceClassesByCapability;
+
+	TMap<FName, TSubclassOf<ULiveLinkDeviceCapability>> TableColumnIdToCapability;
+};
diff --git a/Engine/Plugins/Animation/LiveLinkExampleDevice/LiveLinkExampleDevice.uplugin b/Engine/Plugins/Animation/LiveLinkExampleDevice/LiveLinkExampleDevice.uplugin
new file mode 100644
--- /dev/null
+++ b/Engine/Plugins/Animation/LiveLinkExampleDevice/LiveLinkExampleDevice.uplugin
@@ -0,0 +1,35 @@
+{
+	"FileVersion": 3,
+	"Version": 1,
+	"VersionName": "1.0",
+	"FriendlyName": "Live Link Hub Example Device",
+	"Description": "An example device implementation using the Live Link Hub device interface",
+	"Category": "Animation",
+	"CreatedBy": "Epic Games, Inc.",
+	"CreatedByURL": "https://epicgames.com",
+	"DocsURL": "",
+	"MarketplaceURL": "",
+	"SupportURL": "",
+	"EnabledByDefault": false,
+	"CanContainContent": true,
+	"IsBetaVersion": false,
+	"IsExperimentalVersion": true,
+	"Installed": false,
+	"Plugins": [
+		{
+			"Name": "LiveLinkDevice",
+			"Enabled": true
+		},
+		{
+			"Name": "OSC",
+			"Enabled": true
+		}
+	],
+	"Modules": [
+		{
+			"Name": "LiveLinkExampleDevice",
+			"Type": "Editor",
+			"LoadingPhase": "Default"
+		}
+	]
+}
diff --git a/Engine/Plugins/Animation/LiveLinkExampleDevice/Source/LiveLinkExampleDevice/LiveLinkExampleDevice.Build.cs b/Engine/Plugins/Animation/LiveLinkExampleDevice/Source/LiveLinkExampleDevice/LiveLinkExampleDevice.Build.cs
new file mode 100644
--- /dev/null
+++ b/Engine/Plugins/Animation/LiveLinkExampleDevice/Source/LiveLinkExampleDevice/LiveLinkExampleDevice.Build.cs
@@ -0,0 +1,26 @@
+// Copyright Epic Games, Inc. All Rights Reserved.
+
+using UnrealBuildTool;
+
+public class LiveLinkExampleDevice : ModuleRules
+{
+	public LiveLinkExampleDevice(ReadOnlyTargetRules Target) : base(Target)
+	{
+		PrivateDependencyModuleNames.AddRange(
+			new string[] {
+				"Core",
+				"CoreUObject",
+				"InputCore",
+				"LiveLinkDevice",
+				"LiveLinkHub",
+				"Networking",
+				"OSC",
+				"Projects",
+				"Slate",
+				"SlateCore",
+			}
+		);
+
+		CppCompileWarningSettings.UnsafeTypeCastWarningLevel = WarningLevel.Error;
+	}
+}
diff --git a/Engine/Plugins/Animation/LiveLinkExampleDevice/Source/LiveLinkExampleDevice/Private/Devices/LiveLinkUnrealDevice.cpp b/Engine/Plugins/Animation/LiveLinkExampleDevice/Source/LiveLinkExampleDevice/Private/Devices/LiveLinkUnrealDevice.cpp
new file mode 100644
--- /dev/null
+++ b/Engine/Plugins/Animation/LiveLinkExampleDevice/Source/LiveLinkExampleDevice/Private/Devices/LiveLinkUnrealDevice.cpp
@@ -0,0 +1,126 @@
+// Copyright Epic Games, Inc. All Rights Reserved.
+
+#include "LiveLinkUnrealDevice.h"
+#include "ILiveLinkRecordingSessionInfo.h"
+#include "Logging/StructuredLog.h"
+#include "Modules/ModuleManager.h"
+#include "Widgets/Text/STextBlock.h"
+
+
+IMPLEMENT_MODULE(FDefaultModuleImpl, LiveLinkExampleDevice);
+
+
+DECLARE_LOG_CATEGORY_EXTERN(LogLiveLinkUnrealDevice, Log, All);
+
+DEFINE_LOG_CATEGORY(LogLiveLinkUnrealDevice);
+
+
+#define LOCTEXT_NAMESPACE "LiveLinkUnrealDevice"
+
+
+namespace UE::OscCmds
+{
+	constexpr const TCHAR* RecordStart = TEXT("/RecordStart");
+	constexpr const TCHAR* RecordStop = TEXT("/RecordStop");
+}
+
+
+void ULiveLinkUnrealDevice::OnDeviceAdded(const FGuid InDeviceGuid, ULiveLinkDeviceSettings* InSettings)
+{
+	Super::OnDeviceAdded(InDeviceGuid, InSettings);
+
+	OscClient = UE::OSC::IClientProxy::Create(TEXT("ULiveLinkUnrealDevice"));
+}
+
+
+void ULiveLinkUnrealDevice::OnSettingChanged(const FPropertyChangedEvent& InPropertyChangedEvent)
+{
+	Super::OnSettingChanged(InPropertyChangedEvent);
+
+	const ULiveLinkUnrealDeviceSettings* DeviceSettings = GetSettings();
+
+	const FName ChangedPropertyName = InPropertyChangedEvent.GetPropertyName();
+	static const FName AddressName = GET_MEMBER_NAME_CHECKED(ULiveLinkUnrealDeviceSettings, IPAddress);
+	static const FName OSCPortName = GET_MEMBER_NAME_CHECKED(ULiveLinkUnrealDeviceSettings, OSCPort);
+	if (ChangedPropertyName == AddressName || ChangedPropertyName == OSCPortName)
+	{
+		UE_LOGFMT(LogLiveLinkUnrealDevice, Verbose, "Device '{Name}': Changed address to {Address}:{Port}",
+			DeviceSettings->DisplayName, DeviceSettings->IPAddress, DeviceSettings->OSCPort);
+
+		FIPv4Address Address;
+		ensure(FIPv4Address::Parse(DeviceSettings->IPAddress, Address));
+
+		OscClient->SetSendIPEndpoint(FIPv4Endpoint(Address, DeviceSettings->OSCPort));
+	}
+}
+
+
+ELiveLinkDeviceConnectionStatus ULiveLinkUnrealDevice::GetConnectionStatus_Implementation() const
+{
+	return ELiveLinkDeviceConnectionStatus::Disconnected;
+}
+
+
+FString ULiveLinkUnrealDevice::GetHardwareId_Implementation() const
+{
+	return GetSettings()->IPAddress;
+}
+
+
+bool ULiveLinkUnrealDevice::SetHardwareId_Implementation(const FString& HardwareID)
+{
+	return false;
+}
+
+
+bool ULiveLinkUnrealDevice::Connect_Implementation()
+{
+	return false;
+}
+
+
+bool ULiveLinkUnrealDevice::Disconnect_Implementation()
+{
+	return false;
+}
+
+
+bool ULiveLinkUnrealDevice::StartRecording_Implementation()
+{
+	if (OscClient)
+	{
+		ILiveLinkRecordingSessionInfo& SessionInfo = ILiveLinkRecordingSessionInfo::Get();
+		OscClient->SendMessage(FOSCMessage(
+			FOSCAddress(UE::OscCmds::RecordStart), {
+				UE::OSC::FOSCData(SessionInfo.GetSlateName()),
+				UE::OSC::FOSCData(SessionInfo.GetTakeNumber()),
+				UE::OSC::FOSCData(SessionInfo.GetSessionName()),
+			}));
+	}
+
+	bIsRecording = true;
+
+	return true;
+}
+
+
+bool ULiveLinkUnrealDevice::StopRecording_Implementation()
+{
+	if (OscClient)
+	{
+		OscClient->SendMessage(FOSCMessage(FOSCAddress(UE::OscCmds::RecordStop), {}));
+	}
+
+	bIsRecording = false;
+
+	return true;
+}
+
+
+bool ULiveLinkUnrealDevice::IsRecording_Implementation() const
+{
+	return bIsRecording;
+}
+
+
+#undef LOCTEXT_NAMESPACE
diff --git a/Engine/Plugins/Animation/LiveLinkExampleDevice/Source/LiveLinkExampleDevice/Private/Devices/LiveLinkUnrealDevice.h b/Engine/Plugins/Animation/LiveLinkExampleDevice/Source/LiveLinkExampleDevice/Private/Devices/LiveLinkUnrealDevice.h
new file mode 100644
--- /dev/null
+++ b/Engine/Plugins/Animation/LiveLinkExampleDevice/Source/LiveLinkExampleDevice/Private/Devices/LiveLinkUnrealDevice.h
@@ -0,0 +1,69 @@
+// Copyright Epic Games, Inc. All Rights Reserved.
+
+#pragma once
+
+#include "LiveLinkDevice.h"
+#include "LiveLinkDeviceCapability_Connection.h"
+#include "LiveLinkDeviceCapability_Recording.h"
+#include "OSCClient.h"
+#include "LiveLinkUnrealDevice.generated.h"
+
+
+UCLASS()
+class ULiveLinkUnrealDeviceSettings : public ULiveLinkDeviceSettings
+{
+    GENERATED_BODY()
+
+public:
+	UPROPERTY(EditAnywhere, Category="Unreal Device")
+	FString DisplayName = TEXT("Unreal Editor");
+
+	UPROPERTY(EditAnywhere, Category="Unreal Device")
+	FString IPAddress = TEXT("127.0.0.1");
+
+	UPROPERTY(EditAnywhere, Category="Unreal Device")
+	uint16 OSCPort = 5500;
+};
+
+
+UCLASS()
+class ULiveLinkUnrealDevice : public ULiveLinkDevice
+	, public ILiveLinkDeviceCapability_Connection
+	, public ILiveLinkDeviceCapability_Recording
+{
+	GENERATED_BODY()
+
+public:
+	ULiveLinkUnrealDeviceSettings* GetSettings() { return CastChecked<ULiveLinkUnrealDeviceSettings>(Settings); }
+	const ULiveLinkUnrealDeviceSettings* GetSettings() const { return CastChecked<ULiveLinkUnrealDeviceSettings>(Settings); }
+
+	//~ Begin ULiveLinkDevice interface
+	virtual TSubclassOf<ULiveLinkDeviceSettings> StaticGetSettingsClass() const { return ULiveLinkUnrealDeviceSettings::StaticClass(); }
+
+	virtual FText GetDisplayName() const override { return FText::FromString(GetSettings()->DisplayName); }
+	virtual EDeviceHealth GetDeviceHealth() const override { return EDeviceHealth::Nominal; }
+	virtual FText GetHealthText() const override { return FText::FromString("Example Health"); }
+
+	virtual void OnDeviceAdded(const FGuid InDeviceGuid, ULiveLinkDeviceSettings* InSettings) override;
+	virtual void OnSettingChanged(const FPropertyChangedEvent& InPropertyChangedEvent) override;
+	//~ End ULiveLinkDevice interface
+
+	//~ Begin ILiveLinkDeviceCapability_Connection interface
+	virtual ELiveLinkDeviceConnectionStatus GetConnectionStatus_Implementation() const override;
+	virtual FString GetHardwareId_Implementation() const override;
+	virtual bool SetHardwareId_Implementation(const FString& HardwareID) override;
+	virtual bool Connect_Implementation() override;
+	virtual bool Disconnect_Implementation() override;
+	//~ End ILiveLinkDeviceCapability_Connection interface
+
+	//~ Begin ILiveLinkDeviceCapability_Recording interface
+	virtual bool StartRecording_Implementation() override;
+	virtual bool StopRecording_Implementation() override;
+	virtual bool IsRecording_Implementation() const override;
+	//~ End ILiveLinkDeviceCapability_Recording interface
+
+protected:
+	TUniquePtr<UE::OSC::IClientProxy> OscClient;
+
+	bool bIsRecording;
+};
diff --git a/Engine/Plugins/Animation/LiveLinkExampleDevice/Source/LiveLinkExampleDevice/Private/Tests/LiveLinkUnrealDeviceTests.cpp b/Engine/Plugins/Animation/LiveLinkExampleDevice/Source/LiveLinkExampleDevice/Private/Tests/LiveLinkUnrealDeviceTests.cpp
new file mode 100644
--- /dev/null
+++ b/Engine/Plugins/Animation/LiveLinkExampleDevice/Source/LiveLinkExampleDevice/Private/Tests/LiveLinkUnrealDeviceTests.cpp
@@ -0,0 +1,15 @@
+// Copyright Epic Games, Inc. All Rights Reserved.
+
+#include "Devices/LiveLinkUnrealDevice.h"
+#include "Misc/AutomationTest.h"
+
+
+IMPLEMENT_SIMPLE_AUTOMATION_TEST(FExampleDeviceTest, "LiveLinkHub.UnrealDevice",
+	EAutomationTestFlags::EditorContext | EAutomationTestFlags::EngineFilter)
+
+bool FExampleDeviceTest::RunTest(const FString& Parameters)
+{
+	ULiveLinkUnrealDevice* Device = NewObject<ULiveLinkUnrealDevice>();
+
+	return true;
+}
diff --git a/Engine/Source/Programs/LiveLinkHub/Source/LiveLinkHub.Target.cs b/Engine/Source/Programs/LiveLinkHub/Source/LiveLinkHub.Target.cs
--- a/Engine/Source/Programs/LiveLinkHub/Source/LiveLinkHub.Target.cs
+++ b/Engine/Source/Programs/LiveLinkHub/Source/LiveLinkHub.Target.cs
@@ -57,6 +57,9 @@ public class LiveLinkHubTarget : TargetRules
 			"AppleARKitFaceSupport",
 			"XInputDevice",
 
+			"LiveLinkDevice",
+			"LiveLinkExampleDevice",
+
 			"LiveLinkOpenVR",
 		});
 
