Missing export for public class ULiveLinkDeviceSubsystem

diff --git a/Engine/Plugins/Animation/LiveLinkDevice/Source/LiveLinkDevice/Public/LiveLinkDeviceSubsystem.h b/Engine/Plugins/Animation/LiveLinkDevice/Source/LiveLinkDevice/Public/LiveLinkDeviceSubsystem.h
--- a/Engine/Plugins/Animation/LiveLinkDevice/Source/LiveLinkDevice/Public/LiveLinkDeviceSubsystem.h
+++ b/Engine/Plugins/Animation/LiveLinkDevice/Source/LiveLinkDevice/Public/LiveLinkDeviceSubsystem.h
@@ -16,7 +16,7 @@ class ULiveLinkDeviceCapability;
  * Facilitates cached lookups related to device and capability classes.
  */
 UCLASS(BlueprintType)
-class ULiveLinkDeviceSubsystem : public UEngineSubsystem
+class LIVELINKDEVICE_API ULiveLinkDeviceSubsystem : public UEngineSubsystem
 {
 	GENERATED_BODY()
 
