Resolving an issue when deleting a device from the subsystem

diff --git a/Engine/Plugins/Animation/LiveLinkDevice/Source/LiveLinkDevice/Private/LiveLinkDeviceSubsystem.cpp b/Engine/Plugins/Animation/LiveLinkDevice/Source/LiveLinkDevice/Private/LiveLinkDeviceSubsystem.cpp
--- a/Engine/Plugins/Animation/LiveLinkDevice/Source/LiveLinkDevice/Private/LiveLinkDeviceSubsystem.cpp
+++ b/Engine/Plugins/Animation/LiveLinkDevice/Source/LiveLinkDevice/Private/LiveLinkDeviceSubsystem.cpp
@@ -98,11 +98,13 @@ void ULiveLinkDeviceSubsystem::RemoveDevice(ULiveLinkDevice* InDevice)
 
 	if (const FGuid* MaybeId = Devices.FindKey(InDevice); ensure(MaybeId))
 	{
+		const FGuid DeviceId = *MaybeId;
+
 		UE_LOGFMT(LogLiveLinkDevice, Log,
 			"ULiveLinkDeviceSubsystem: Removing device {DeviceName} with ID {DeviceGuid}.",
-			InDevice->GetFName(), *MaybeId);
-		Devices.Remove(*MaybeId);
-		OnDeviceRemovedDelegate.Broadcast(*MaybeId, InDevice);
+			InDevice->GetFName(), DeviceId);
+		Devices.Remove(DeviceId);
+		OnDeviceRemovedDelegate.Broadcast(DeviceId, InDevice);
 	}
 	else
 	{
