Missing OnDeviceRemoved call on the LiveLink device

diff --git a/Engine/Plugins/Animation/LiveLinkDevice/Source/LiveLinkDevice/Private/LiveLinkDeviceSubsystem.cpp b/Engine/Plugins/Animation/LiveLinkDevice/Source/LiveLinkDevice/Private/LiveLinkDeviceSubsystem.cpp
--- a/Engine/Plugins/Animation/LiveLinkDevice/Source/LiveLinkDevice/Private/LiveLinkDeviceSubsystem.cpp
+++ b/Engine/Plugins/Animation/LiveLinkDevice/Source/LiveLinkDevice/Private/LiveLinkDeviceSubsystem.cpp
@@ -76,12 +76,15 @@ void ULiveLinkDeviceSubsystem::RemoveDevice(FGuid InDeviceId)
 {
 	if (TObjectPtr<ULiveLinkDevice>* MaybeDevice = Devices.Find(InDeviceId); ensure(MaybeDevice))
 	{
+		TObjectPtr<ULiveLinkDevice> Device = *MaybeDevice;
 		UE_LOGFMT(LogLiveLinkDevice, Log,
 			"ULiveLinkDeviceSubsystem: Removing device {DeviceName} with ID {DeviceGuid}.",
-			MaybeDevice->GetFName(), InDeviceId);
-		DevicesByClass.Remove(MaybeDevice->GetClass(), *MaybeDevice);
+			Device->GetFName(), InDeviceId);
+		DevicesByClass.Remove(Device->GetClass(), Device);
 		Devices.Remove(InDeviceId);
-		OnDeviceRemovedDelegate.Broadcast(InDeviceId, *MaybeDevice);
+
+		Device->OnDeviceRemoved();
+		OnDeviceRemovedDelegate.Broadcast(InDeviceId, Device);
 	}
 	else
 	{
@@ -104,6 +107,8 @@ void ULiveLinkDeviceSubsystem::RemoveDevice(ULiveLinkDevice* InDevice)
 			"ULiveLinkDeviceSubsystem: Removing device {DeviceName} with ID {DeviceGuid}.",
 			InDevice->GetFName(), DeviceId);
 		Devices.Remove(DeviceId);
+
+		InDevice->OnDeviceRemoved();
 		OnDeviceRemovedDelegate.Broadcast(DeviceId, InDevice);
 	}
 	else
