Changes to the Connection capability

diff --git a/Engine/Plugins/Animation/LiveLinkDevice/Source/LiveLinkDevice/Private/LiveLinkDeviceCapability_Connection.cpp b/Engine/Plugins/Animation/LiveLinkDevice/Source/LiveLinkDevice/Private/LiveLinkDeviceCapability_Connection.cpp
--- a/Engine/Plugins/Animation/LiveLinkDevice/Source/LiveLinkDevice/Private/LiveLinkDeviceCapability_Connection.cpp
+++ b/Engine/Plugins/Animation/LiveLinkDevice/Source/LiveLinkDevice/Private/LiveLinkDeviceCapability_Connection.cpp
@@ -27,3 +27,19 @@ TSharedPtr<SWidget> ULiveLinkDeviceCapability_Connection::GenerateWidgetForColum
 			}
 		);
 }
+
+ILiveLinkDeviceCapability_Connection::ILiveLinkDeviceCapability_Connection()
+	: ConnectionDelegate(NewObject<UConnectionDelegate>())
+{
+}
+
+UConnectionDelegate* ILiveLinkDeviceCapability_Connection::GetConnectionDelegate_Implementation()
+{
+	return ConnectionDelegate.Get();
+}
+
+void ILiveLinkDeviceCapability_Connection::SetConnectionStatus(ELiveLinkDeviceConnectionStatus InStatus)
+{
+	ConnectionDelegate->ConnectionChanged.Broadcast(InStatus);
+	ConnectionDelegate->ConnectionChangedDynamic.Broadcast(InStatus);
+}
\ No newline at end of file
diff --git a/Engine/Plugins/Animation/LiveLinkDevice/Source/LiveLinkDevice/Public/LiveLinkDeviceCapability_Connection.h b/Engine/Plugins/Animation/LiveLinkDevice/Source/LiveLinkDevice/Public/LiveLinkDeviceCapability_Connection.h
--- a/Engine/Plugins/Animation/LiveLinkDevice/Source/LiveLinkDevice/Public/LiveLinkDeviceCapability_Connection.h
+++ b/Engine/Plugins/Animation/LiveLinkDevice/Source/LiveLinkDevice/Public/LiveLinkDeviceCapability_Connection.h
@@ -45,12 +45,30 @@ public:
 	virtual TSharedPtr<SWidget> GenerateWidgetForColumn(const FName InColumnId, ULiveLinkDevice* InDevice) const override;
 };
 
+DECLARE_TS_MULTICAST_DELEGATE_OneParam(FDeviceConnectionStatusChanged, ELiveLinkDeviceConnectionStatus);
+UDELEGATE()
+DECLARE_DYNAMIC_MULTICAST_DELEGATE_OneParam(FDeviceConnectionStatusChangedDynamic, ELiveLinkDeviceConnectionStatus, InNewStatus);
+
+UCLASS()
+class LIVELINKDEVICE_API UConnectionDelegate : public UObject
+{
+	GENERATED_BODY()
+
+public:
+
+	UPROPERTY(EditAnywhere, Category = "Live Link Device|Connection")
+	FDeviceConnectionStatusChangedDynamic ConnectionChangedDynamic;
+
+	FDeviceConnectionStatusChanged ConnectionChanged;
+};
 
 class LIVELINKDEVICE_API ILiveLinkDeviceCapability_Connection : public ILiveLinkDeviceCapability
 {
 	GENERATED_BODY()
 
 public:
+	ILiveLinkDeviceCapability_Connection();
+
 	UFUNCTION(BlueprintCallable, BlueprintNativeEvent, Category="Live Link Device|Connection")
 	ELiveLinkDeviceConnectionStatus GetConnectionStatus() const;
 	virtual ELiveLinkDeviceConnectionStatus GetConnectionStatus_Implementation() const PURE_VIRTUAL(ILiveLinkDeviceCapability_Connection::GetConnectionStatus_Implementation, return ELiveLinkDeviceConnectionStatus::Disconnected;);
@@ -72,4 +90,17 @@ public:
 	UFUNCTION(BlueprintCallable, BlueprintNativeEvent, Category="Live Link Device|Connection")
 	bool Disconnect();
 	virtual bool Disconnect_Implementation() PURE_VIRTUAL(ILiveLinkDeviceCapability_Connection::Disconnect_Implementation, return false;);
+
+	UFUNCTION(BlueprintCallable, BlueprintNativeEvent, Category="Live Link Device|Connection")
+	UConnectionDelegate* GetConnectionDelegate();
+
+	UConnectionDelegate* GetConnectionDelegate_Implementation();
+
+protected:
+
+	void SetConnectionStatus(ELiveLinkDeviceConnectionStatus InStatus);
+
+private:
+
+	TStrongObjectPtr<UConnectionDelegate> ConnectionDelegate;
 };
