Hiding abstract devices from the menu

diff --git a/Engine/Plugins/Animation/LiveLinkDevice/Source/LiveLinkDevice/Private/LiveLinkDeviceModule.cpp b/Engine/Plugins/Animation/LiveLinkDevice/Source/LiveLinkDevice/Private/LiveLinkDeviceModule.cpp
--- a/Engine/Plugins/Animation/LiveLinkDevice/Source/LiveLinkDevice/Private/LiveLinkDeviceModule.cpp
+++ b/Engine/Plugins/Animation/LiveLinkDevice/Source/LiveLinkDevice/Private/LiveLinkDeviceModule.cpp
@@ -89,6 +89,11 @@ TSharedRef<SWidget> FLiveLinkDeviceModule::OnGenerateAddDeviceMenu()
 
 	for (const TSubclassOf<ULiveLinkDevice>& DeviceClass : DeviceClasses)
 	{
+		if (DeviceClass->HasAnyClassFlags(CLASS_Abstract))
+		{
+			continue;
+		}
+
 		const ULiveLinkDevice* DeviceCDO = DeviceClass->GetDefaultObject<ULiveLinkDevice>();
 		if (!ensure(DeviceCDO))
 		{
@@ -96,6 +101,12 @@ TSharedRef<SWidget> FLiveLinkDeviceModule::OnGenerateAddDeviceMenu()
 		}
 
 		TSubclassOf<ULiveLinkDeviceSettings> DeviceSettingsClass = DeviceCDO->StaticGetSettingsClass();
+
+		if (DeviceSettingsClass->HasAnyClassFlags(CLASS_Abstract))
+		{
+			continue;
+		}
+
 		MenuBuilder.AddMenuEntry(
 			DeviceClass->GetDisplayNameText(),
 			FText::FromString(DeviceClass->GetPathName()),
diff --git a/Engine/Plugins/Animation/LiveLinkDevice/Source/LiveLinkDevice/Private/Tests/LiveLinkDeviceTests.h b/Engine/Plugins/Animation/LiveLinkDevice/Source/LiveLinkDevice/Private/Tests/LiveLinkDeviceTests.h
--- a/Engine/Plugins/Animation/LiveLinkDevice/Source/LiveLinkDevice/Private/Tests/LiveLinkDeviceTests.h
+++ b/Engine/Plugins/Animation/LiveLinkDevice/Source/LiveLinkDevice/Private/Tests/LiveLinkDeviceTests.h
@@ -28,6 +28,11 @@ public:
 	virtual void SetValue_Implementation(int32 InValue) PURE_VIRTUAL(ILiveLinkDeviceCapability_BaseTest::SetValue_Implementation, );
 };
 
+UCLASS()
+class ULiveLinkDeviceSettings_BaseTest : public ULiveLinkDeviceSettings
+{
+	GENERATED_BODY()
+};
 
 UCLASS()
 class ULiveLinkDevice_BaseTest : public ULiveLinkDevice
@@ -36,6 +41,7 @@ class ULiveLinkDevice_BaseTest : public ULiveLinkDevice
 	GENERATED_BODY()
 
 public:
+	virtual TSubclassOf<ULiveLinkDeviceSettings> StaticGetSettingsClass() const { return ULiveLinkDeviceSettings_BaseTest::StaticClass(); }
 	//~ Begin ULiveLinkDevice interface
 	virtual FText GetDisplayName() const override { return FText::FromString("Test Device"); }
 	virtual EDeviceHealth GetDeviceHealth() const override { return EDeviceHealth::Nominal; }
