Live Link device framework: Device session serialization, capability UI improvements, code cleanup.

ULiveLinkDeviceSubsystem:
- `ULiveLinkDeviceSubsystem` now implements `ILiveLinkHubSessionExtraDataHandler`, enabling it to read/write device configurations and recreate them on session load.
- Device creation is now better encapsulated.

ULiveLinkDevice:
- Replaced `ULiveLinkDevice::GetBaseSettings` with `GetDeviceSettings<T>`, which handles casting to a derived type (defaulting to the base), and made the raw `Settings` pointer private.
- `ULiveLinkDevice` now stores the Device GUID internally, accessible via `GetDeviceId()`.
- `ULiveLinkDevice::OnDeviceAdded` is no longer responsible for chaining to the base class.
- Renamed `StaticGetSettingsClass` to `GetSettingsClass`.
- `NotPlaceable` UCLASS meta is now usable as an alternative to `Abstract` in order to hide device types from being enumerated in the creation UI, since `Abstract` precludes instantiation for e.g. automated tests.

Capabilities:
- Capability widget generation now receives a delegate indicating whether their device row is selected.
- Capability widget generation is now non-const, to facilitate patterns like calling `ULiveLinkDeviceCapability_Connection::ConnectAllDevices()`.
- Connection capability "hardware ID" table widget is now directly editable if a device opts in by returning true from `CanSetHardwareId`.
- Connection capability now serves as an example of adding a menu of actions to its column header, which will be refactored later to additionally support context menu actions.
- Recording capability status column now correctly renders with icons from the newly-implemented Slate style, `FLiveLinkDeviceStyle`.
- Simplified several unnecessary IInterface uses of `PURE_VIRTUAL` macro to actual pure virtuals.

LiveLinkExampleDevice:
- Simplified example device code, and moved the updated LiveLinkUnrealDevice to a new LiveLinkStandardDevices plugin (forthcoming in a separate CL).

CaptureManager:
- API updates. Mostly light touch now that `UGenericLiveLinkDeviceHelper` has been removed.

#rnx

diff --git a/Engine/Plugins/Animation/LiveLinkDevice/Source/LiveLinkDevice/Private/LiveLinkDevice.cpp b/Engine/Plugins/Animation/LiveLinkDevice/Source/LiveLinkDevice/Private/LiveLinkDevice.cpp
--- a/Engine/Plugins/Animation/LiveLinkDevice/Source/LiveLinkDevice/Private/LiveLinkDevice.cpp
+++ b/Engine/Plugins/Animation/LiveLinkDevice/Source/LiveLinkDevice/Private/LiveLinkDevice.cpp
@@ -10,41 +10,8 @@
 #include "Widgets/SNullWidget.h"
 
 
-TSharedRef<FJsonObject> ULiveLinkDeviceSettings::SaveToJson() const
-{
-	TSharedRef<FJsonObject> JsonObject = MakeShared<FJsonObject>();
-	FJsonObjectConverter::UStructToJsonObject(GetClass(), this, JsonObject);
-	return JsonObject;
-}
-
-
-/** Called to restore device settings from session configuration. */
-bool ULiveLinkDeviceSettings::RestoreFromJson(const TSharedRef<FJsonObject>& InJsonObject)
-{
-	FText FailReason;
-
-	const int64 CheckFlags = 0;
-	const int64 SkipFlags = 0;
-	const bool bStrictMode_false = false;
-	const bool bJsonConverted = FJsonObjectConverter::JsonObjectToUStruct(
-		InJsonObject, GetClass(), this, CheckFlags, SkipFlags, bStrictMode_false, &FailReason
-	);
-
-	if (bJsonConverted)
-	{
-		return true;
-	}
-
-	UE_LOGFMT(LogLiveLinkDevice, Error,
-		"ULiveLinkDeviceSettings::RestoreFromJson: JsonObjectToUStruct failed ({Reason})",
-		FailReason.ToString()
-	);
-
-	return false;
-}
-
 
-TSharedRef<SWidget> ULiveLinkDevice::GenerateWidgetForColumn(const FName InColumnId)
+TSharedRef<SWidget> ULiveLinkDevice::GenerateWidgetForColumn(const FName InColumnId, const FLiveLinkDeviceWidgetArguments& InArgs)
 {
 	// Check if the capability provides a default widget implementation.
 	ULiveLinkDeviceSubsystem* Subsystem = GEngine->GetEngineSubsystem<ULiveLinkDeviceSubsystem>();
@@ -52,8 +19,8 @@ TSharedRef<SWidget> ULiveLinkDevice::GenerateWidgetForColumn(const FName InColum
 	const TSubclassOf<ULiveLinkDeviceCapability>* MaybeCapabilityClass = ColumnIdToCapability.Find(InColumnId);
 	if (ensure(MaybeCapabilityClass))
 	{
-		const ULiveLinkDeviceCapability* CapabilityCDO = MaybeCapabilityClass->GetDefaultObject();
-		if (TSharedPtr<SWidget> CapabilityWidget = CapabilityCDO->GenerateWidgetForColumn(InColumnId, this))
+		ULiveLinkDeviceCapability* CapabilityCDO = MaybeCapabilityClass->GetDefaultObject();
+		if (TSharedPtr<SWidget> CapabilityWidget = CapabilityCDO->GenerateWidgetForColumn(InColumnId, InArgs, this))
 		{
 			return CapabilityWidget.ToSharedRef();
 		}
@@ -65,19 +32,10 @@ TSharedRef<SWidget> ULiveLinkDevice::GenerateWidgetForColumn(const FName InColum
 }
 
 
-void ULiveLinkDevice::OnDeviceAdded(const FGuid InDeviceGuid, ULiveLinkDeviceSettings* InSettings)
+void ULiveLinkDevice::InternalDeviceAdded(const FGuid InDeviceGuid, ULiveLinkDeviceSettings* InSettings)
 {
+	DeviceGuid = InDeviceGuid;
 	Settings = InSettings;
-}
-
-
-void ULiveLinkDevice::OnDeviceRemoved()
-{
-
-}
-
-
-void ULiveLinkDevice::OnSettingChanged(const FPropertyChangedEvent& InPropertyChangedEvent)
-{
 
+	OnDeviceAdded();
 }
diff --git a/Engine/Plugins/Animation/LiveLinkDevice/Source/LiveLinkDevice/Private/LiveLinkDeviceCapability_Connection.cpp b/Engine/Plugins/Animation/LiveLinkDevice/Source/LiveLinkDevice/Private/LiveLinkDeviceCapability_Connection.cpp
--- a/Engine/Plugins/Animation/LiveLinkDevice/Source/LiveLinkDevice/Private/LiveLinkDeviceCapability_Connection.cpp
+++ b/Engine/Plugins/Animation/LiveLinkDevice/Source/LiveLinkDevice/Private/LiveLinkDeviceCapability_Connection.cpp
@@ -1,16 +1,92 @@
 // Copyright Epic Games, Inc. All Rights Reserved.
 
 #include "LiveLinkDeviceCapability_Connection.h"
+#include "Engine/Engine.h"
+#include "Framework/MultiBox/MultiBoxBuilder.h"
 #include "LiveLinkDevice.h"
+#include "LiveLinkDeviceSubsystem.h"
+#include "Styling/StyleColors.h"
+#include "Widgets/SBoxPanel.h"
+#include "Widgets/SCompoundWidget.h"
+#include "Widgets/Input/SButton.h"
+#include "Widgets/Text/SInlineEditableTextBlock.h"
 #include "Widgets/Text/STextBlock.h"
 
 
+#define LOCTEXT_NAMESPACE "LiveLinkDevice"
+
+
+class SLiveLinkConnection_HardwareId : public SCompoundWidget
+{
+	SLATE_BEGIN_ARGS(SLiveLinkConnection_HardwareId) {}
+		SLATE_ATTRIBUTE(FText, Text)
+		SLATE_ATTRIBUTE(bool, IsReadOnly)
+
+		/** Must be propagated for editable child widgets to work. */
+		SLATE_EVENT(FIsSelected, IsDeviceRowSelected)
+
+		SLATE_EVENT(FOnTextCommitted, OnTextCommitted)
+	SLATE_END_ARGS()
+
+public:
+	void Construct(const FArguments& InArgs)
+	{
+		ChildSlot
+		[
+			SNew(SInlineEditableTextBlock)
+			.IsReadOnly(InArgs._IsReadOnly)
+			.IsSelected(InArgs._IsDeviceRowSelected)
+			.OnTextCommitted(InArgs._OnTextCommitted)
+			.Text(InArgs._Text)
+		];
+	}
+
+protected:
+	TAttribute<ELiveLinkDeviceConnectionStatus> ConnectionStatus;
+};
+
+
+SHeaderRow::FColumn::FArguments& ULiveLinkDeviceCapability_Connection::GenerateHeaderForColumn(
+	const FName InColumnId,
+	SHeaderRow::FColumn::FArguments& InArgs
+)
+{
+	if (InColumnId == Column_HardwareId)
+	{
+		return InArgs
+			.DefaultLabel(LOCTEXT("HardwareID_ColumnHeader_Label", "Hardware ID"))
+			.DefaultTooltip(LOCTEXT("HardwareID_ColumnHeader_Tooltip", "Device hardware identifier"))
+			.HeaderComboVisibility(EHeaderComboVisibility::Ghosted)
+			.OnGetMenuContent_UObject(this, &ULiveLinkDeviceCapability_Connection::Header_GetMenuContent)
+			.FillWidth(0.75f);
+	}
+
+	return Super::GenerateHeaderForColumn(InColumnId, InArgs);
+}
+
+
 TSharedPtr<SWidget> ULiveLinkDeviceCapability_Connection::GenerateWidgetForColumn(
 	const FName InColumnId,
+	const FLiveLinkDeviceWidgetArguments& InArgs,
 	ULiveLinkDevice* InDevice
-) const
+)
 {
-	return SNew(STextBlock)
+	return SNew(SLiveLinkConnection_HardwareId)
+		.IsDeviceRowSelected(InArgs.IsRowSelected)
+		.IsReadOnly_Lambda(
+			[WeakDevice = TWeakObjectPtr<ULiveLinkDevice>(InDevice)]
+			()
+			{
+				if (ULiveLinkDevice* Device = WeakDevice.Get())
+				{
+					return !ILiveLinkDeviceCapability_Connection::Execute_CanSetHardwareId(Device);
+				}
+				else
+				{
+					return true;
+				}
+			}
+		)
 		.Text_Lambda(
 			[WeakDevice = TWeakObjectPtr<ULiveLinkDevice>(InDevice)]
 			()
@@ -25,21 +101,92 @@ TSharedPtr<SWidget> ULiveLinkDeviceCapability_Connection::GenerateWidgetForColum
 					return FText();
 				}
 			}
+		)
+		.OnTextCommitted_Lambda(
+			[WeakDevice = TWeakObjectPtr<ULiveLinkDevice>(InDevice)]
+			(const FText& InNewText, ETextCommit::Type InCommitType)
+			{
+				if (ULiveLinkDevice* Device = WeakDevice.Get())
+				{
+					ILiveLinkDeviceCapability_Connection::Execute_SetHardwareId(Device, InNewText.ToString());
+				}
+			}
 		);
 }
 
+
+TSharedRef<SWidget> ULiveLinkDeviceCapability_Connection::Header_GetMenuContent()
+{
+	const bool bShouldCloseAfterSelection_true = true;
+	FMenuBuilder MenuBuilder(bShouldCloseAfterSelection_true, nullptr);
+
+	MenuBuilder.AddMenuEntry(
+		LOCTEXT("HardwareID_ColumnHeader_ConnectAll_Label", "Connect All Devices"),
+		TAttribute<FText>(),
+		FSlateIcon(),
+		FUIAction(FExecuteAction::CreateUObject(this, &ULiveLinkDeviceCapability_Connection::ConnectAllDevices))
+	);
+
+	MenuBuilder.AddMenuEntry(
+		LOCTEXT("HardwareID_ColumnHeader_DisconnectAll_Label", "Disconnect All Devices"),
+		TAttribute<FText>(),
+		FSlateIcon(),
+		FUIAction(FExecuteAction::CreateUObject(this, &ULiveLinkDeviceCapability_Connection::DisconnectAllDevices))
+	);
+
+	return MenuBuilder.MakeWidget();
+}
+
+
+void ULiveLinkDeviceCapability_Connection::ConnectAllDevices()
+{
+	ULiveLinkDeviceSubsystem* Subsystem = GEngine->GetEngineSubsystem<ULiveLinkDeviceSubsystem>();
+
+	TArray<ULiveLinkDevice*> ConnectibleDevices;
+	Subsystem->GetDevicesByCapability(ULiveLinkDeviceCapability_Connection::StaticClass(), ConnectibleDevices);
+	for (ULiveLinkDevice* Device : ConnectibleDevices)
+	{
+		if (ensure(Device->Implements<ULiveLinkDeviceCapability_Connection>()))
+		{
+			ILiveLinkDeviceCapability_Connection::Execute_Connect(Device);
+		}
+	}
+}
+
+
+void ULiveLinkDeviceCapability_Connection::DisconnectAllDevices()
+{
+	ULiveLinkDeviceSubsystem* Subsystem = GEngine->GetEngineSubsystem<ULiveLinkDeviceSubsystem>();
+
+	TArray<ULiveLinkDevice*> ConnectibleDevices;
+	Subsystem->GetDevicesByCapability(ULiveLinkDeviceCapability_Connection::StaticClass(), ConnectibleDevices);
+	for (ULiveLinkDevice* Device : ConnectibleDevices)
+	{
+		if (ensure(Device->Implements<ULiveLinkDeviceCapability_Connection>()))
+		{
+			ILiveLinkDeviceCapability_Connection::Execute_Disconnect(Device);
+		}
+	}
+}
+
+
 ILiveLinkDeviceCapability_Connection::ILiveLinkDeviceCapability_Connection()
 	: ConnectionDelegate(NewObject<UConnectionDelegate>())
 {
 }
 
+
 UConnectionDelegate* ILiveLinkDeviceCapability_Connection::GetConnectionDelegate_Implementation()
 {
 	return ConnectionDelegate.Get();
 }
 
+
 void ILiveLinkDeviceCapability_Connection::SetConnectionStatus(ELiveLinkDeviceConnectionStatus InStatus)
 {
 	ConnectionDelegate->ConnectionChanged.Broadcast(InStatus);
 	ConnectionDelegate->ConnectionChangedDynamic.Broadcast(InStatus);
-}
\ No newline at end of file
+}
+
+
+#undef LOCTEXT_NAMESPACE
diff --git a/Engine/Plugins/Animation/LiveLinkDevice/Source/LiveLinkDevice/Private/LiveLinkDeviceCapability_Recording.cpp b/Engine/Plugins/Animation/LiveLinkDevice/Source/LiveLinkDevice/Private/LiveLinkDeviceCapability_Recording.cpp
--- a/Engine/Plugins/Animation/LiveLinkDevice/Source/LiveLinkDevice/Private/LiveLinkDeviceCapability_Recording.cpp
+++ b/Engine/Plugins/Animation/LiveLinkDevice/Source/LiveLinkDevice/Private/LiveLinkDeviceCapability_Recording.cpp
@@ -4,8 +4,47 @@
 #include "Engine/Engine.h"
 #include "ILiveLinkRecordingSessionInfo.h"
 #include "LiveLinkDevice.h"
+#include "LiveLinkDeviceStyle.h"
 #include "LiveLinkDeviceSubsystem.h"
-#include "Widgets/Text/STextBlock.h"
+#include "Widgets/Images/SImage.h"
+
+
+class SLiveLinkCapability_Recording : public SCompoundWidget
+{
+	SLATE_BEGIN_ARGS(SLiveLinkCapability_Recording) {}
+		SLATE_ATTRIBUTE(bool, IsRecording)
+	SLATE_END_ARGS()
+
+public:
+	void Construct(const FArguments& InArgs)
+	{
+		IsRecording = InArgs._IsRecording;
+
+		ChildSlot
+		.VAlign(VAlign_Center)
+		.HAlign(HAlign_Center)
+		[
+			SNew(SImage)
+			.Image_Raw(this, &SLiveLinkCapability_Recording::GetImage)
+			.ColorAndOpacity_Raw(this, &SLiveLinkCapability_Recording::GetColorAndOpacity)
+		];
+	}
+
+	const FSlateBrush* GetImage() const
+	{
+		return IsRecording.Get()
+			? FLiveLinkDeviceStyle::Get()->GetBrush("Record")
+			: FLiveLinkDeviceStyle::Get()->GetBrush("Record.Monochrome");
+	}
+
+	FSlateColor GetColorAndOpacity() const
+	{
+		return IsRecording.Get() ? FSlateColor::UseForeground() : FSlateColor::UseSubduedForeground();
+	}
+
+protected:
+	TAttribute<bool> IsRecording;
+};
 
 
 void ULiveLinkDeviceCapability_Recording::OnDeviceSubsystemInitialized()
@@ -38,11 +77,11 @@ void ULiveLinkDeviceCapability_Recording::OnDeviceSubsystemDeinitializing()
 SHeaderRow::FColumn::FArguments& ULiveLinkDeviceCapability_Recording::GenerateHeaderForColumn(
 	const FName InColumnId,
 	SHeaderRow::FColumn::FArguments& InArgs
-) const
+)
 {
 	if (InColumnId == Column_RecordingStatus)
 	{
-		return InArgs.DefaultLabel(FText::FromString("Rec"))
+		return InArgs.DefaultLabel(FText::FromString(" "))
 			.DefaultTooltip(FText::FromString("Device recording status"))
 			.FillSized(30.0f);
 	}
@@ -53,23 +92,21 @@ SHeaderRow::FColumn::FArguments& ULiveLinkDeviceCapability_Recording::GenerateHe
 
 TSharedPtr<SWidget> ULiveLinkDeviceCapability_Recording::GenerateWidgetForColumn(
 	const FName InColumnId,
+	const FLiveLinkDeviceWidgetArguments& InArgs,
 	ULiveLinkDevice* InDevice
-) const
+)
 {
-	return SNew(STextBlock)
-		.Text_Lambda(
+	return SNew(SLiveLinkCapability_Recording)
+		.IsRecording_Lambda(
 			[WeakDevice = TWeakObjectPtr<ULiveLinkDevice>(InDevice)]
 			()
 			{
 				if (ULiveLinkDevice* Device = WeakDevice.Get())
 				{
-					const bool bRecording = ILiveLinkDeviceCapability_Recording::Execute_IsRecording(Device);
-					return bRecording ? FText::FromString(TEXT("R")) : FText::FromString(TEXT("-"));
-				}
-				else
-				{
-					return FText();
+					return ILiveLinkDeviceCapability_Recording::Execute_IsRecording(Device);
 				}
+
+				return false;
 			}
 		);
 }
diff --git a/Engine/Plugins/Animation/LiveLinkDevice/Source/LiveLinkDevice/Private/LiveLinkDeviceModule.cpp b/Engine/Plugins/Animation/LiveLinkDevice/Source/LiveLinkDevice/Private/LiveLinkDeviceModule.cpp
--- a/Engine/Plugins/Animation/LiveLinkDevice/Source/LiveLinkDevice/Private/LiveLinkDeviceModule.cpp
+++ b/Engine/Plugins/Animation/LiveLinkDevice/Source/LiveLinkDevice/Private/LiveLinkDeviceModule.cpp
@@ -6,6 +6,7 @@
 #include "Framework/Docking/TabManager.h"
 #include "Framework/MultiBox/MultiBoxBuilder.h"
 #include "LiveLinkDevice.h"
+#include "LiveLinkDeviceStyle.h"
 #include "LiveLinkDeviceSubsystem.h"
 #include "Logging/StructuredLog.h"
 #include "PropertyEditorModule.h"
@@ -32,6 +33,8 @@ IMPLEMENT_MODULE(FLiveLinkDeviceModule, LiveLinkDevice);
 
 void FLiveLinkDeviceModule::StartupModule()
 {
+	FLiveLinkDeviceStyle::Initialize();
+
 	FGlobalTabmanager::Get()
 		->RegisterNomadTabSpawner(DevicesTabName,
 			FOnSpawnTab::CreateRaw(this, &FLiveLinkDeviceModule::OnSpawnDevicesTab))
@@ -46,7 +49,7 @@ void FLiveLinkDeviceModule::StartupModule()
 
 void FLiveLinkDeviceModule::ShutdownModule()
 {
-	OnDeviceSelectionChangedDelegate.RemoveAll(this);
+	FLiveLinkDeviceStyle::Shutdown();
 }
 
 
@@ -56,7 +59,7 @@ void FLiveLinkDeviceModule::DeviceSelectionChanged(ULiveLinkDevice* InSelectedDe
 
 	if (DetailsView)
 	{
-		DetailsView->SetObject(InSelectedDevice ? InSelectedDevice->GetBaseSettings() : nullptr);
+		DetailsView->SetObject(InSelectedDevice ? InSelectedDevice->GetDeviceSettings() : nullptr);
 	}
 
 	OnDeviceSelectionChangedDelegate.Broadcast(InSelectedDevice);
@@ -76,20 +79,7 @@ TSharedRef<SWidget> FLiveLinkDeviceModule::OnGenerateAddDeviceMenu()
 
 	for (const TSubclassOf<ULiveLinkDevice>& DeviceClass : DeviceClasses)
 	{
-		if (DeviceClass->HasAnyClassFlags(CLASS_Abstract))
-		{
-			continue;
-		}
-
-		const ULiveLinkDevice* DeviceCDO = DeviceClass->GetDefaultObject<ULiveLinkDevice>();
-		if (!ensure(DeviceCDO))
-		{
-			continue;
-		}
-
-		TSubclassOf<ULiveLinkDeviceSettings> DeviceSettingsClass = DeviceCDO->StaticGetSettingsClass();
-
-		if (DeviceSettingsClass->HasAnyClassFlags(CLASS_Abstract))
+		if (DeviceClass->HasAnyClassFlags(CLASS_Abstract | CLASS_NotPlaceable))
 		{
 			continue;
 		}
@@ -101,14 +91,7 @@ TSharedRef<SWidget> FLiveLinkDeviceModule::OnGenerateAddDeviceMenu()
 			FUIAction(FExecuteAction::CreateLambda(
 				[DeviceSubsystem, DeviceClass]
 				{
-					UE_LOGFMT(LogLiveLinkDevice, Log, "Creating new Live Link Device of class {DeviceClass}", DeviceClass->GetName());
-					ULiveLinkDevice* NewDevice = NewObject<ULiveLinkDevice>(DeviceSubsystem, DeviceClass);
-					const FGuid DeviceGuid = DeviceSubsystem->AddDevice(NewDevice);
-
-					TSubclassOf<ULiveLinkDeviceSettings> SettingsClass = NewDevice->StaticGetSettingsClass();
-					UE_LOGFMT(LogLiveLinkDevice, Log, "Initializing with {DeviceSettingsClass}", SettingsClass->GetName());
-					ULiveLinkDeviceSettings* NewSettings = NewObject<ULiveLinkDeviceSettings>(NewDevice, SettingsClass);
-					NewDevice->OnDeviceAdded(DeviceGuid, NewSettings);
+					DeviceSubsystem->CreateDeviceOfClass(DeviceClass);
 				}
 			)),
 			NAME_None,
@@ -191,7 +174,7 @@ TSharedRef<SDockTab> FLiveLinkDeviceModule::OnSpawnDeviceDetailsTab(const FSpawn
 	DetailsView = PropertyEditorModule.CreateDetailView(DetailsViewArgs);
 	if (ULiveLinkDevice* SelectedDevice = WeakSelectedDevice.Get())
 	{
-		DetailsView->SetObject(SelectedDevice->GetBaseSettings());
+		DetailsView->SetObject(SelectedDevice->GetDeviceSettings());
 	}
 
 	DetailsView->OnFinishedChangingProperties().AddLambda(
diff --git a/Engine/Plugins/Animation/LiveLinkDevice/Source/LiveLinkDevice/Private/LiveLinkDeviceStyle.cpp b/Engine/Plugins/Animation/LiveLinkDevice/Source/LiveLinkDevice/Private/LiveLinkDeviceStyle.cpp
new file mode 100644
--- /dev/null
+++ b/Engine/Plugins/Animation/LiveLinkDevice/Source/LiveLinkDevice/Private/LiveLinkDeviceStyle.cpp
@@ -0,0 +1,44 @@
+// Copyright Epic Games, Inc. All Rights Reserved.
+
+#include "LiveLinkDeviceStyle.h"
+#include "Interfaces/IPluginManager.h"
+#include "Styling/SlateStyle.h"
+#include "Styling/SlateStyleRegistry.h"
+
+
+TSharedPtr<FSlateStyleSet> FLiveLinkDeviceStyle::StyleSet;
+
+
+void FLiveLinkDeviceStyle::Initialize()
+{
+	if (!ensureMsgf(!StyleSet, TEXT("FLiveLinkDeviceStyle already initialized")))
+	{
+		return;
+	}
+
+	StyleSet = MakeShared<FSlateStyleSet>(GetStyleSetName());
+
+	const FString PluginContentRoot =
+		IPluginManager::Get().FindPlugin("LiveLinkDevice")->GetBaseDir() / TEXT("Resources");
+	StyleSet->SetContentRoot(PluginContentRoot);
+
+	static const FVector2D Icon16x16(16.0f, 16.0f);
+	static const FVector2D Icon20x20(20.0f, 20.0f);
+	StyleSet->Set("Record", new FSlateVectorImageBrush(PluginContentRoot / "PlayControlsRecord.svg", Icon16x16));
+	StyleSet->Set("Record.Monochrome", new FSlateVectorImageBrush(PluginContentRoot / "PlayControlsRecord_Monochrome.svg", Icon16x16));
+
+	FSlateStyleRegistry::RegisterSlateStyle(*StyleSet);
+}
+
+
+void FLiveLinkDeviceStyle::Shutdown()
+{
+	FSlateStyleRegistry::UnRegisterSlateStyle(*StyleSet);
+	StyleSet.Reset();
+}
+
+
+TSharedPtr<ISlateStyle> FLiveLinkDeviceStyle::Get()
+{
+	return StyleSet;
+}
diff --git a/Engine/Plugins/Animation/LiveLinkDevice/Source/LiveLinkDevice/Private/LiveLinkDeviceStyle.h b/Engine/Plugins/Animation/LiveLinkDevice/Source/LiveLinkDevice/Private/LiveLinkDeviceStyle.h
new file mode 100644
--- /dev/null
+++ b/Engine/Plugins/Animation/LiveLinkDevice/Source/LiveLinkDevice/Private/LiveLinkDeviceStyle.h
@@ -0,0 +1,27 @@
+// Copyright Epic Games, Inc. All Rights Reserved.
+
+#pragma once
+
+#include "Templates/SharedPointer.h"
+
+class FSlateStyleSet;
+class ISlateStyle;
+
+
+class FLiveLinkDeviceStyle
+{
+public:
+	static void Initialize();
+	static void Shutdown();
+
+	static TSharedPtr<ISlateStyle> Get();
+
+	static FName GetStyleSetName()
+	{
+		static const FName StyleSetName("LiveLinkDeviceStyle");
+		return StyleSetName;
+	}
+
+private:
+	static TSharedPtr<FSlateStyleSet> StyleSet;
+};
diff --git a/Engine/Plugins/Animation/LiveLinkDevice/Source/LiveLinkDevice/Private/LiveLinkDeviceSubsystem.cpp b/Engine/Plugins/Animation/LiveLinkDevice/Source/LiveLinkDevice/Private/LiveLinkDeviceSubsystem.cpp
--- a/Engine/Plugins/Animation/LiveLinkDevice/Source/LiveLinkDevice/Private/LiveLinkDeviceSubsystem.cpp
+++ b/Engine/Plugins/Animation/LiveLinkDevice/Source/LiveLinkDevice/Private/LiveLinkDeviceSubsystem.cpp
@@ -5,12 +5,15 @@
 #include "LiveLinkDevice.h"
 #include "LiveLinkDeviceCapability.h"
 #include "LiveLinkDeviceModule.h"
+#include "LiveLinkHubSessionExtraData_Device.h"
 
 
 void ULiveLinkDeviceSubsystem::Initialize(FSubsystemCollectionBase& Collection)
 {
 	Super::Initialize(Collection);
 
+	RegisterExtraDataHandler();
+
 	// Enumerate capability classes
 	{
 		TArray<UClass*> Capabilities;
@@ -46,6 +49,8 @@ void ULiveLinkDeviceSubsystem::Initialize(FSubsystemCollectionBase& Collection)
 
 void ULiveLinkDeviceSubsystem::Deinitialize()
 {
+	UnregisterExtraDataHandler();
+
 	// Dispatch de-initializing calls to capability CDOs.
 	for (TSubclassOf<ULiveLinkDeviceCapability>& Capability : KnownCapabilities)
 	{
@@ -53,22 +58,200 @@ void ULiveLinkDeviceSubsystem::Deinitialize()
 		CapabilityCDO->OnDeviceSubsystemDeinitializing();
 	}
 
+	// Remove all existing devices.
+	TArray<FGuid> DeviceIds;
+	Devices.GetKeys(DeviceIds);
+	for (const FGuid& DeviceId : DeviceIds)
+	{
+		RemoveDevice(DeviceId);
+	}
+
 	Super::Deinitialize();
 }
 
 
-FGuid ULiveLinkDeviceSubsystem::AddDevice(ULiveLinkDevice* InDevice)
+TSubclassOf<ULiveLinkHubSessionExtraData> ULiveLinkDeviceSubsystem::GetExtraDataClass() const
+{
+	return ULiveLinkHubSessionExtraData_Device::StaticClass();
+}
+
+
+void ULiveLinkDeviceSubsystem::OnExtraDataSessionSaving(ULiveLinkHubSessionExtraData* InExtraData)
+{
+	ULiveLinkHubSessionExtraData_Device* DeviceData = CastChecked<ULiveLinkHubSessionExtraData_Device>(InExtraData);
+
+	DeviceData->Devices.Empty();
+
+	for (const TPair<FGuid, TObjectPtr<ULiveLinkDevice>>& DevicePair : Devices)
+	{
+		DeviceData->Devices.Emplace(FLiveLinkDevicePreset{
+			.DeviceGuid = DevicePair.Key,
+			.DeviceClass = DevicePair.Value->GetClass(),
+			.DeviceSettings = DevicePair.Value->GetDeviceSettings(),
+		});
+	}
+}
+
+
+void ULiveLinkDeviceSubsystem::OnExtraDataSessionLoaded(const ULiveLinkHubSessionExtraData* InExtraData)
+{
+	// Remove all existing devices.
+	{
+		TArray<FGuid> AllDeviceGuids;
+		AllDeviceGuids.Reserve(Devices.Num());
+		Devices.GetKeys(AllDeviceGuids);
+
+		for (const FGuid& DeviceGuid : AllDeviceGuids)
+		{
+			RemoveDevice(DeviceGuid);
+		}
+	}
+
+	if (InExtraData == nullptr)
+	{
+		// New/loaded session has no saved device data.
+		return;
+	}
+
+	const ULiveLinkHubSessionExtraData_Device* DeviceData = CastChecked<ULiveLinkHubSessionExtraData_Device>(InExtraData);
+
+	for (const FLiveLinkDevicePreset& DevicePreset : DeviceData->Devices)
+	{
+		UE_LOGFMT(LogLiveLinkDevice, Log, "Restoring saved Live Link device with ID {Guid}", DevicePreset.DeviceGuid);
+
+		if (!DevicePreset.DeviceClass)
+		{
+			UE_LOGFMT(LogLiveLinkDevice, Error, "Device class missing");
+			continue;
+		}
+
+		if (!DevicePreset.DeviceSettings)
+		{
+			UE_LOGFMT(LogLiveLinkDevice, Error, "Device settings missing");
+			continue;
+		}
+
+		FCreateResult CreateResult = InternalCreateDeviceOfClass(
+			DevicePreset.DeviceClass,
+			DevicePreset.DeviceSettings,
+			DevicePreset.DeviceGuid
+		);
+
+		if (!CreateResult.HasValue())
+		{
+			UE_LOGFMT(LogLiveLinkDevice, Error, "Failed to restore device");
+		}
+	}
+}
+
+
+ULiveLinkDeviceSubsystem::FCreateResult ULiveLinkDeviceSubsystem::CreateDeviceOfClass(
+	TSubclassOf<ULiveLinkDevice> InDeviceClass,
+	ULiveLinkDeviceSettings* InDeviceSettingsTemplate /* = nullptr */
+)
+{
+	return InternalCreateDeviceOfClass(InDeviceClass, InDeviceSettingsTemplate);
+}
+
+
+ULiveLinkDeviceSubsystem::FCreateResult ULiveLinkDeviceSubsystem::InternalCreateDeviceOfClass(
+	TSubclassOf<ULiveLinkDevice> InDeviceClass,
+	ULiveLinkDeviceSettings* InDeviceSettingsTemplate /* = nullptr */,
+	TOptional<FGuid> InDeviceGuid /* = {} */
+)
+{
+	UE_LOGFMT(LogLiveLinkDevice, Log, "Creating new Live Link device of class {DeviceClass}", InDeviceClass->GetName());
+
+	if (!InDeviceGuid)
+	{
+		InDeviceGuid = FGuid::NewGuid();
+	}
+	else if (InDeviceGuid == FGuid())
+	{
+		InDeviceGuid = FGuid::NewGuid();
+		UE_LOGFMT(LogLiveLinkDevice, Warning, "Tried to create device with unset GUID; assigning new GUID {DeviceGuid}", InDeviceGuid);
+	}
+
+	const FGuid& DeviceGuid = InDeviceGuid.GetValue();
+
+	ULiveLinkDevice* NewDevice = NewObject<ULiveLinkDevice>(this, InDeviceClass);
+	if (!ensure(NewDevice))
+	{
+		UE_LOGFMT(LogLiveLinkDevice, Error, "Failed to create ULiveLinkDevice");
+		return MakeError();
+	}
+
+	TSubclassOf<ULiveLinkDeviceSettings> SettingsClass = NewDevice->GetSettingsClass();
+	ULiveLinkDeviceSettings* NewSettings = nullptr;
+	if (InDeviceSettingsTemplate)
+	{
+		TSubclassOf<ULiveLinkDeviceSettings> TemplateClass = InDeviceSettingsTemplate->GetClass();
+		if (TemplateClass != SettingsClass)
+		{
+			UE_LOGFMT(LogLiveLinkDevice, Error,
+				"Settings template is of wrong class (got {TemplateClass}, expected {DeviceClass})",
+				TemplateClass->GetName(),
+				SettingsClass->GetName());
+
+			return MakeError();
+		}
+
+		NewSettings = DuplicateObject(InDeviceSettingsTemplate, NewDevice);
+	}
+	else
+	{
+		NewSettings = NewObject<ULiveLinkDeviceSettings>(NewDevice, SettingsClass);
+	}
+
+	if (!ensure(NewSettings))
+	{
+		UE_LOGFMT(LogLiveLinkDevice, Error, "Failed to create ULiveLinkDeviceSettings");
+		return MakeError();
+	}
+
+	InternalAddDevice(DeviceGuid, NewDevice, NewSettings);
+
+	return MakeValue(DeviceGuid, NewDevice);
+}
+
+
+void ULiveLinkDeviceSubsystem::InternalAddDevice(
+	FGuid InGuid,
+	ULiveLinkDevice* InDevice,
+	ULiveLinkDeviceSettings* InSettings
+)
 {
 	check(InDevice);
+	check(!Devices.Find(InGuid));
 
-	const FGuid DeviceGuid = FGuid::NewGuid();
 	UE_LOGFMT(LogLiveLinkDevice, Log,
-		"ULiveLinkDeviceSubsystem: Adding device {DeviceName} with ID {DeviceGuid}.",
-		InDevice->GetFName(), DeviceGuid);
-	Devices.Add(DeviceGuid, InDevice);
+		"ULiveLinkDeviceSubsystem: Adding device {DeviceName} with ID {InGuid}.",
+		InDevice->GetFName(), InGuid);
+
+	Devices.Add(InGuid, InDevice);
 	DevicesByClass.Add(InDevice->GetClass(), InDevice);
-	OnDeviceAddedDelegate.Broadcast(DeviceGuid, InDevice);
-	return DeviceGuid;
+
+	InDevice->InternalDeviceAdded(InGuid, InSettings);
+
+	OnDeviceAddedDelegate.Broadcast(InGuid, InDevice);
+}
+
+void ULiveLinkDeviceSubsystem::InternalRemoveDevice(FGuid InDeviceId, ULiveLinkDevice* InDevice)
+{
+	check(InDevice);
+	check(Devices.Find(InDeviceId));
+	check(*Devices.Find(InDeviceId) == InDevice);
+
+	UE_LOGFMT(LogLiveLinkDevice, Log,
+		"ULiveLinkDeviceSubsystem: Removing device {DeviceName} with ID {DeviceGuid}.",
+		InDevice->GetFName(), InDeviceId);
+
+	Devices.Remove(InDeviceId);
+	DevicesByClass.Remove(InDevice->GetClass(), InDevice);
+
+	InDevice->OnDeviceRemoved();
+
+	OnDeviceRemovedDelegate.Broadcast(InDeviceId, InDevice);
 }
 
 
@@ -77,14 +260,7 @@ void ULiveLinkDeviceSubsystem::RemoveDevice(FGuid InDeviceId)
 	if (TObjectPtr<ULiveLinkDevice>* MaybeDevice = Devices.Find(InDeviceId); ensure(MaybeDevice))
 	{
 		TObjectPtr<ULiveLinkDevice> Device = *MaybeDevice;
-		UE_LOGFMT(LogLiveLinkDevice, Log,
-			"ULiveLinkDeviceSubsystem: Removing device {DeviceName} with ID {DeviceGuid}.",
-			Device->GetFName(), InDeviceId);
-		DevicesByClass.Remove(Device->GetClass(), Device);
-		Devices.Remove(InDeviceId);
-
-		Device->OnDeviceRemoved();
-		OnDeviceRemovedDelegate.Broadcast(InDeviceId, Device);
+		InternalRemoveDevice(InDeviceId, Device);
 	}
 	else
 	{
@@ -102,14 +278,7 @@ void ULiveLinkDeviceSubsystem::RemoveDevice(ULiveLinkDevice* InDevice)
 	if (const FGuid* MaybeId = Devices.FindKey(InDevice); ensure(MaybeId))
 	{
 		const FGuid DeviceId = *MaybeId;
-
-		UE_LOGFMT(LogLiveLinkDevice, Log,
-			"ULiveLinkDeviceSubsystem: Removing device {DeviceName} with ID {DeviceGuid}.",
-			InDevice->GetFName(), DeviceId);
-		Devices.Remove(DeviceId);
-
-		InDevice->OnDeviceRemoved();
-		OnDeviceRemovedDelegate.Broadcast(DeviceId, InDevice);
+		InternalRemoveDevice(DeviceId, InDevice);
 	}
 	else
 	{
diff --git a/Engine/Plugins/Animation/LiveLinkDevice/Source/LiveLinkDevice/Private/LiveLinkHubSessionExtraData_Device.h b/Engine/Plugins/Animation/LiveLinkDevice/Source/LiveLinkDevice/Private/LiveLinkHubSessionExtraData_Device.h
new file mode 100644
--- /dev/null
+++ b/Engine/Plugins/Animation/LiveLinkDevice/Source/LiveLinkDevice/Private/LiveLinkHubSessionExtraData_Device.h
@@ -0,0 +1,38 @@
+// Copyright Epic Games, Inc. All Rights Reserved.
+
+#pragma once
+
+#include "LiveLinkHubSessionExtraData.h"
+#include "Templates/SubclassOf.h"
+#include "LiveLinkHubSessionExtraData_Device.generated.h"
+
+
+class ULiveLinkDevice;
+class ULiveLinkDeviceSettings;
+
+
+USTRUCT()
+struct FLiveLinkDevicePreset
+{
+	GENERATED_BODY()
+
+	UPROPERTY()
+	FGuid DeviceGuid;
+
+	UPROPERTY()
+	TSubclassOf<ULiveLinkDevice> DeviceClass;
+
+	UPROPERTY(Instanced)
+	TObjectPtr<ULiveLinkDeviceSettings> DeviceSettings;
+};
+
+
+UCLASS()
+class ULiveLinkHubSessionExtraData_Device : public ULiveLinkHubSessionExtraData
+{
+	GENERATED_BODY()
+
+public:
+	UPROPERTY()
+	TArray<FLiveLinkDevicePreset> Devices;
+};
diff --git a/Engine/Plugins/Animation/LiveLinkDevice/Source/LiveLinkDevice/Private/Tests/LiveLinkDevice.spec.cpp b/Engine/Plugins/Animation/LiveLinkDevice/Source/LiveLinkDevice/Private/Tests/LiveLinkDevice.spec.cpp
new file mode 100644
--- /dev/null
+++ b/Engine/Plugins/Animation/LiveLinkDevice/Source/LiveLinkDevice/Private/Tests/LiveLinkDevice.spec.cpp
@@ -0,0 +1,88 @@
+// Copyright Epic Games, Inc. All Rights Reserved.
+
+#if WITH_DEV_AUTOMATION_TESTS
+
+#include "Engine/Engine.h"
+#include "LiveLinkDevice_BasicTest.h"
+#include "LiveLinkDeviceSubsystem.h"
+#include "Misc/AutomationTest.h"
+//#include "Tests/Assertions.h"
+#include "UObject/ScriptInterface.h"
+
+
+BEGIN_DEFINE_SPEC(FLiveLinkDeviceSpec, "LiveLinkHub.Devices",
+                  EAutomationTestFlags::EditorContext | EAutomationTestFlags::ProductFilter)
+	ULiveLinkDevice_BasicTest* TestDevice;
+END_DEFINE_SPEC(FLiveLinkDeviceSpec);
+
+
+void FLiveLinkDeviceSpec::Define()
+{
+	Describe("ULiveLinkDevice_BasicTest", [this]
+	{
+		BeforeEach([this]
+		{
+			ULiveLinkDeviceSubsystem* Subsystem = GEngine->GetEngineSubsystem<ULiveLinkDeviceSubsystem>();
+			ULiveLinkDeviceSubsystem::FCreateResult CreateResult =
+				Subsystem->CreateDeviceOfClass(ULiveLinkDevice_BasicTest::StaticClass());
+			TestDevice = CastChecked<ULiveLinkDevice_BasicTest>(CreateResult.GetValue().Device);
+		});
+
+		AfterEach([this]
+		{
+			ULiveLinkDeviceSubsystem* Subsystem = GEngine->GetEngineSubsystem<ULiveLinkDeviceSubsystem>();
+			Subsystem->RemoveDevice(TestDevice);
+			TestDevice = nullptr;
+		});
+
+		It("can be queried for capabilities via base class pointer", [this]
+		{
+			ULiveLinkDevice* UnknownDevice = TestDevice;
+
+			// Can still query and retrieve the implemented capabilities.
+			TestTrueExpr(UnknownDevice->Implements<ULiveLinkDeviceCapability_BasicTest>());
+
+			TScriptInterface<ILiveLinkDeviceCapability_BasicTest> TestCapability =
+				TScriptInterface<ILiveLinkDeviceCapability_BasicTest>(UnknownDevice);
+
+			TestNotNull(
+				TEXT("TScriptInterface<ILiveLinkDeviceCapability_BasicTest>::GetInterface()"),
+				TestCapability.GetInterface()
+			);
+		});
+
+		It("can handle capability method invocations via base class pointer", [this]
+		{
+			ULiveLinkDevice* UnknownDevice = TestDevice;
+
+			const int32 RandomValue = FMath::Rand32();
+			ILiveLinkDeviceCapability_BasicTest::Execute_SetValue(UnknownDevice, RandomValue);
+			TestEqual(
+				TEXT("Randomly generated value set through interface"),
+				ILiveLinkDeviceCapability_BasicTest::Execute_GetValue(UnknownDevice),
+				RandomValue
+			);
+		});
+	});
+
+	Describe("ULiveLinkDeviceSubsystem", [this]
+	{
+		It("cannot create a device with a settings template of the wrong type", [this]
+		{
+			ULiveLinkDeviceSubsystem* Subsystem = GEngine->GetEngineSubsystem<ULiveLinkDeviceSubsystem>();
+
+			// Try to create a device with the wrong settings subclass. This should fail.
+			AddExpectedMessage(TEXT("Settings template is of wrong class"), ELogVerbosity::Error);
+
+			ULiveLinkDeviceSettings_Invalid* InvalidSettings = NewObject<ULiveLinkDeviceSettings_Invalid>();
+
+			ULiveLinkDeviceSubsystem::FCreateResult InvalidSettingsCreateResult =
+				Subsystem->CreateDeviceOfClass(ULiveLinkDevice_BasicTest::StaticClass(), InvalidSettings);
+
+			TestTrueExpr(InvalidSettingsCreateResult.HasError());
+		});
+	});
+}
+
+
+#endif // #if WITH_DEV_AUTOMATION_TESTS
diff --git a/Engine/Plugins/Animation/LiveLinkDevice/Source/LiveLinkDevice/Private/Tests/LiveLinkDeviceTests.cpp b/Engine/Plugins/Animation/LiveLinkDevice/Source/LiveLinkDevice/Private/Tests/LiveLinkDeviceTests.cpp
deleted file mode 100644
diff --git a/Engine/Plugins/Animation/LiveLinkDevice/Source/LiveLinkDevice/Private/Tests/LiveLinkDeviceTests.h b/Engine/Plugins/Animation/LiveLinkDevice/Source/LiveLinkDevice/Private/Tests/LiveLinkDeviceTests.h
deleted file mode 100644
diff --git a/Engine/Plugins/Animation/LiveLinkDevice/Source/LiveLinkDevice/Private/Tests/LiveLinkDevice_BasicTest.h b/Engine/Plugins/Animation/LiveLinkDevice/Source/LiveLinkDevice/Private/Tests/LiveLinkDevice_BasicTest.h
new file mode 100644
--- /dev/null
+++ b/Engine/Plugins/Animation/LiveLinkDevice/Source/LiveLinkDevice/Private/Tests/LiveLinkDevice_BasicTest.h
@@ -0,0 +1,67 @@
+// Copyright Epic Games, Inc. All Rights Reserved.
+
+#pragma once
+
+#include "LiveLinkDevice.h"
+#include "LiveLinkDeviceCapability.h"
+#include "LiveLinkDevice_BasicTest.generated.h"
+
+
+UINTERFACE()
+class LIVELINKDEVICE_API ULiveLinkDeviceCapability_BasicTest : public ULiveLinkDeviceCapability
+{
+	GENERATED_BODY()
+};
+
+
+class LIVELINKDEVICE_API ILiveLinkDeviceCapability_BasicTest : public ILiveLinkDeviceCapability
+{
+	GENERATED_BODY()
+
+public:
+	UFUNCTION(BlueprintCallable, BlueprintNativeEvent, Category="Live Link Device|Testing")
+	int32 GetValue() const;
+	virtual int32 GetValue_Implementation() const = 0;
+
+	UFUNCTION(BlueprintCallable, BlueprintNativeEvent, Category="Live Link Device|Testing")
+	void SetValue(int32 InValue);
+	virtual void SetValue_Implementation(int32 InValue) = 0;
+};
+
+
+UCLASS()
+class ULiveLinkDeviceSettings_BasicTest : public ULiveLinkDeviceSettings
+{
+	GENERATED_BODY()
+};
+
+UCLASS(NotPlaceable)
+class ULiveLinkDevice_BasicTest : public ULiveLinkDevice
+	, public ILiveLinkDeviceCapability_BasicTest
+{
+	GENERATED_BODY()
+
+public:
+	//~ Begin ULiveLinkDevice interface
+	virtual TSubclassOf<ULiveLinkDeviceSettings> GetSettingsClass() const override { return ULiveLinkDeviceSettings_BasicTest::StaticClass(); }
+	virtual FText GetDisplayName() const override { return FText::FromString("Test Device"); }
+	virtual EDeviceHealth GetDeviceHealth() const override { return EDeviceHealth::Nominal; }
+	virtual FText GetHealthText() const override { return FText::FromString("Test Health"); }
+	//~ End ULiveLinkDevice interface
+
+	//~ Begin ILiveLinkDeviceCapability_BaseTest interface
+	virtual int32 GetValue_Implementation() const override { return TestValue; }
+	virtual void SetValue_Implementation(int32 InValue) override { TestValue = InValue; }
+	//~ End ILiveLinkDeviceCapability_BaseTest interface
+
+protected:
+	int32 TestValue;
+};
+
+
+/** A type that does not correspond to any device, and is used to test expected failures. */
+UCLASS()
+class ULiveLinkDeviceSettings_Invalid : public ULiveLinkDeviceSettings
+{
+	GENERATED_BODY()
+};
diff --git a/Engine/Plugins/Animation/LiveLinkDevice/Source/LiveLinkDevice/Private/Widgets/DeviceTable.cpp b/Engine/Plugins/Animation/LiveLinkDevice/Source/LiveLinkDevice/Private/Widgets/DeviceTable.cpp
--- a/Engine/Plugins/Animation/LiveLinkDevice/Source/LiveLinkDevice/Private/Widgets/DeviceTable.cpp
+++ b/Engine/Plugins/Animation/LiveLinkDevice/Source/LiveLinkDevice/Private/Widgets/DeviceTable.cpp
@@ -127,7 +127,7 @@ void SLiveLinkDeviceTable::RegenerateColumns()
 	ULiveLinkDeviceSubsystem* Subsystem = GEngine->GetEngineSubsystem<ULiveLinkDeviceSubsystem>();
 	for (const TSubclassOf<ULiveLinkDeviceCapability>& Capability : Subsystem->GetKnownCapabilities())
 	{
-		const ULiveLinkDeviceCapability* CapabilityCDO = Capability.GetDefaultObject();
+		ULiveLinkDeviceCapability* CapabilityCDO = Capability.GetDefaultObject();
 		for (const TPair<FName, ULiveLinkDeviceCapability::FDeviceTableColumnDesc>& Column : CapabilityCDO->GetTableColumns())
 		{
 			const FName ColumnId = Column.Key;
@@ -188,9 +188,15 @@ TSharedRef<SWidget> SLiveLinkDeviceRowWidget::GenerateWidgetForColumn(const FNam
 
 	if (InColumnId == SLiveLinkDeviceTable::ColumnName_DisplayName)
 	{
-		return SNew(STextBlock)
-			.Text_Lambda([Device]() { return Device->GetDisplayName(); })
-			.ToolTipText_Lambda([Device]() { return Device->GetDisplayName(); });
+		return
+			SNew(SBorder)
+			.VAlign(VAlign_Center)
+			.BorderImage(FAppStyle::GetBrush("NoBorder"))
+			[
+				SNew(STextBlock)
+				.Text_Lambda([Device]() { return Device->GetDisplayName(); })
+				.ToolTipText_Lambda([Device]() { return Device->GetDisplayName(); })
+			];
 	}
 
 	if (InColumnId == SLiveLinkDeviceTable::ColumnName_Remove)
@@ -232,7 +238,11 @@ TSharedRef<SWidget> SLiveLinkDeviceRowWidget::GenerateWidgetForColumn(const FNam
 		return SNullWidget::NullWidget;
 	}
 
-	return Device->GenerateWidgetForColumn(InColumnId);
+	FLiveLinkDeviceWidgetArguments Args = {
+		.IsRowSelected = TDelegate<bool()>::CreateSPLambda(this, [this]() { return IsSelected(); })
+	};
+
+	return Device->GenerateWidgetForColumn(InColumnId, Args);
 }
 
 
diff --git a/Engine/Plugins/Animation/LiveLinkDevice/Source/LiveLinkDevice/Public/LiveLinkDevice.h b/Engine/Plugins/Animation/LiveLinkDevice/Source/LiveLinkDevice/Public/LiveLinkDevice.h
--- a/Engine/Plugins/Animation/LiveLinkDevice/Source/LiveLinkDevice/Public/LiveLinkDevice.h
+++ b/Engine/Plugins/Animation/LiveLinkDevice/Source/LiveLinkDevice/Public/LiveLinkDevice.h
@@ -7,7 +7,9 @@
 
 
 class FJsonObject;
+struct FLiveLinkDeviceWidgetArguments;
 class SWidget;
+class ULiveLinkDeviceSubsystem;
 
 
 UENUM(BlueprintType)
@@ -24,21 +26,11 @@ UCLASS(Abstract)
 class LIVELINKDEVICE_API ULiveLinkDeviceSettings : public UObject
 {
 	GENERATED_BODY()
-
-	/** Called to persist device settings to session configuration. */
-	virtual TSharedRef<FJsonObject> SaveToJson() const;
-
-	/** Called to restore device settings from session configuration. */
-	virtual bool RestoreFromJson(const TSharedRef<FJsonObject>& InJsonObject);
 };
 
 
 /**
  * Abstract base class for all Live Link devices.
- * 
- * Methods having names beginning with "Static" are expected to return the
- * same result regardless of what object instance they are invoked upon,
- * and must be valid to call on the class default object (CDO).
  */
 UCLASS(Abstract, Blueprintable)
 class LIVELINKDEVICE_API ULiveLinkDevice : public UObject
@@ -47,9 +39,8 @@ class LIVELINKDEVICE_API ULiveLinkDevice : public UObject
 
 public:
 	/** Get the device settings class to be displayed in the details view. */
-	virtual TSubclassOf<ULiveLinkDeviceSettings> StaticGetSettingsClass() const { return ULiveLinkDeviceSettings::StaticClass(); }
+	virtual TSubclassOf<ULiveLinkDeviceSettings> GetSettingsClass() const PURE_VIRTUAL(ULiveLinkDevice::GetSettingsClass, return ULiveLinkDeviceSettings::StaticClass(););
 
-public:
 	/** Human-readable display name for this device. */
 	UFUNCTION(BlueprintCallable, Category="Live Link Device")
 	virtual FText GetDisplayName() const PURE_VIRTUAL(ULiveLinkDevice::GetDisplayName, return FText(););
@@ -63,25 +54,41 @@ public:
 	virtual FText GetHealthText() const PURE_VIRTUAL(ULiveLinkDevice::GetHealthText, return FText(););
 
 	/**
-	 * Generate the Slate content, for this device, for the specified column.
+	 * Generate the Slate content, for this device's row, for the specified column.
 	 * The default implementation will delegate to the capability CDO where appropriate.
 	 */
-	virtual TSharedRef<SWidget> GenerateWidgetForColumn(const FName InColumnId);
+	virtual TSharedRef<SWidget> GenerateWidgetForColumn(
+		const FName InColumnId,
+		const FLiveLinkDeviceWidgetArguments& InArgs
+	);
 
 public:
 	/** Called when the device has been registered with the device manager, with either initial or restored settings. */
-	virtual void OnDeviceAdded(const FGuid InDeviceGuid, ULiveLinkDeviceSettings* InSettings);
+	virtual void OnDeviceAdded() { }
 
 	/** Called when the device is removed from the device manager. */
-	virtual void OnDeviceRemoved();
+	virtual void OnDeviceRemoved() { }
 
 	/** Called when a setting value has changed via the UI. */
-	virtual void OnSettingChanged(const FPropertyChangedEvent& InPropertyChangedEvent);
+	virtual void OnSettingChanged(const FPropertyChangedEvent& InPropertyChangedEvent) { }
+
+	/** This device's settings UObject. */
+	template<typename T = ULiveLinkDeviceSettings>
+	const T* GetDeviceSettings() const { return CastChecked<T>(Settings); }
+
+	/** This device's settings UObject (mutable). */
+	template<typename T = ULiveLinkDeviceSettings>
+	T* GetDeviceSettings() { return CastChecked<T>(Settings); }
+
+	/** The identifier with which this device was registered with `ULiveLinkDeviceSubsystem`. */
+	FGuid GetDeviceId() const { return DeviceGuid; }
+
+private:
+	friend ULiveLinkDeviceSubsystem;
+	virtual void InternalDeviceAdded(const FGuid InDeviceGuid, ULiveLinkDeviceSettings* InSettings);
 
-	const ULiveLinkDeviceSettings* GetBaseSettings() const { return Settings; }
-	ULiveLinkDeviceSettings* GetBaseSettings() { return Settings; }
+	FGuid DeviceGuid;
 
-protected:
 	UPROPERTY()
 	TObjectPtr<ULiveLinkDeviceSettings> Settings;
 };
diff --git a/Engine/Plugins/Animation/LiveLinkDevice/Source/LiveLinkDevice/Public/LiveLinkDeviceCapability.h b/Engine/Plugins/Animation/LiveLinkDevice/Source/LiveLinkDevice/Public/LiveLinkDeviceCapability.h
--- a/Engine/Plugins/Animation/LiveLinkDevice/Source/LiveLinkDevice/Public/LiveLinkDeviceCapability.h
+++ b/Engine/Plugins/Animation/LiveLinkDevice/Source/LiveLinkDevice/Public/LiveLinkDeviceCapability.h
@@ -12,6 +12,13 @@ class ULiveLinkDevice;
 class ULiveLinkDeviceCapability;
 
 
+/** Helper passed to GenerateWidgetForColumn. */
+struct FLiveLinkDeviceWidgetArguments
+{
+	TDelegate<bool()> IsRowSelected;
+};
+
+
 /**
  * Base class for all device capabilities.
  *
@@ -49,7 +56,10 @@ public:
 	}
 
 	/** Configure the header for the specified column. */
-	virtual SHeaderRow::FColumn::FArguments& GenerateHeaderForColumn(const FName InColumnId, SHeaderRow::FColumn::FArguments& InArgs) const
+	virtual SHeaderRow::FColumn::FArguments& GenerateHeaderForColumn(
+		const FName InColumnId,
+		SHeaderRow::FColumn::FArguments& InArgs
+	)
 	{
 		// This base class implementation should not get called.
 		// If your capability defines table columns, you need to implement this.
@@ -58,7 +68,11 @@ public:
 	}
 
 	/** Optional; allows a capability to provide a default widget the device can fall back to. */
-	virtual TSharedPtr<SWidget> GenerateWidgetForColumn(const FName InColumnId, ULiveLinkDevice* InDevice) const
+	virtual TSharedPtr<SWidget> GenerateWidgetForColumn(
+		const FName InColumnId,
+		const FLiveLinkDeviceWidgetArguments& InArgs,
+		ULiveLinkDevice* InDevice
+	)
 	{
 		return nullptr;
 	}
diff --git a/Engine/Plugins/Animation/LiveLinkDevice/Source/LiveLinkDevice/Public/LiveLinkDeviceCapability_Connection.h b/Engine/Plugins/Animation/LiveLinkDevice/Source/LiveLinkDevice/Public/LiveLinkDeviceCapability_Connection.h
--- a/Engine/Plugins/Animation/LiveLinkDevice/Source/LiveLinkDevice/Public/LiveLinkDeviceCapability_Connection.h
+++ b/Engine/Plugins/Animation/LiveLinkDevice/Source/LiveLinkDevice/Public/LiveLinkDeviceCapability_Connection.h
@@ -6,6 +6,7 @@
 #include "LiveLinkDeviceCapability_Connection.generated.h"
 
 
+/** Device connection states. */
 UENUM()
 enum class ELiveLinkDeviceConnectionStatus : uint8
 {
@@ -16,6 +17,7 @@ enum class ELiveLinkDeviceConnectionStatus : uint8
 };
 
 
+/** \see ILiveLinkDeviceCapability_Connection. */
 UINTERFACE()
 class LIVELINKDEVICE_API ULiveLinkDeviceCapability_Connection : public ULiveLinkDeviceCapability
 {
@@ -29,26 +31,36 @@ public:
 	{
 	}
 
-	virtual SHeaderRow::FColumn::FArguments& GenerateHeaderForColumn(const FName InColumnId, SHeaderRow::FColumn::FArguments& InArgs) const override
-	{
-		if (InColumnId == Column_HardwareId)
-		{
-			return InArgs
-				.DefaultLabel(FText::FromString("Hardware ID"))
-				.DefaultTooltip(FText::FromString("Device hardware identifier"))
-				.FillWidth(0.5f);
-		}
-
-		return Super::GenerateHeaderForColumn(InColumnId, InArgs);
-	}
+	//~ Begin ULiveLinkDeviceCapability interface
+	virtual SHeaderRow::FColumn::FArguments& GenerateHeaderForColumn(
+		const FName InColumnId,
+		SHeaderRow::FColumn::FArguments& InArgs
+	) override;
 
-	virtual TSharedPtr<SWidget> GenerateWidgetForColumn(const FName InColumnId, ULiveLinkDevice* InDevice) const override;
+	virtual TSharedPtr<SWidget> GenerateWidgetForColumn(
+		const FName InColumnId,
+		const FLiveLinkDeviceWidgetArguments& InArgs,
+		ULiveLinkDevice* InDevice
+	) override;
+	//~ End ULiveLinkDeviceCapability interface
+
+protected:
+	/** Builds the menu for the column header. */
+	TSharedRef<SWidget> Header_GetMenuContent();
+
+	/** Invokes the `Connect` method on all registered devices implementing this interface. */
+	void ConnectAllDevices();
+
+	/** Invokes the `Disconnect` method on all registered devices implementing this interface. */
+	void DisconnectAllDevices();
 };
 
+
 DECLARE_TS_MULTICAST_DELEGATE_OneParam(FDeviceConnectionStatusChanged, ELiveLinkDeviceConnectionStatus);
 UDELEGATE()
 DECLARE_DYNAMIC_MULTICAST_DELEGATE_OneParam(FDeviceConnectionStatusChangedDynamic, ELiveLinkDeviceConnectionStatus, InNewStatus);
 
+
 UCLASS()
 class LIVELINKDEVICE_API UConnectionDelegate : public UObject
 {
@@ -62,6 +74,10 @@ public:
 	FDeviceConnectionStatusChanged ConnectionChanged;
 };
 
+/**
+ * Status and operations relevant to devices with the concept of being (dis)connected.
+ * This could be a physical hardware connection, a network connection, or something else.
+ */
 class LIVELINKDEVICE_API ILiveLinkDeviceCapability_Connection : public ILiveLinkDeviceCapability
 {
 	GENERATED_BODY()
@@ -69,38 +85,43 @@ class LIVELINKDEVICE_API ILiveLinkDeviceCapability_Connection : public ILiveLink
 public:
 	ILiveLinkDeviceCapability_Connection();
 
+	/** Get the current connection state. */
 	UFUNCTION(BlueprintCallable, BlueprintNativeEvent, Category="Live Link Device|Connection")
 	ELiveLinkDeviceConnectionStatus GetConnectionStatus() const;
-	virtual ELiveLinkDeviceConnectionStatus GetConnectionStatus_Implementation() const PURE_VIRTUAL(ILiveLinkDeviceCapability_Connection::GetConnectionStatus_Implementation, return ELiveLinkDeviceConnectionStatus::Disconnected;);
+	virtual ELiveLinkDeviceConnectionStatus GetConnectionStatus_Implementation() const = 0;
 
 	/** Retrieve hardware identifier (serial number, network endpoint, etc). */
 	UFUNCTION(BlueprintCallable, BlueprintNativeEvent, Category="Live Link Device|Connection")
 	FString GetHardwareId() const;
-	virtual FString GetHardwareId_Implementation() const PURE_VIRTUAL(ILiveLinkDeviceCapability_Connection::GetHardwareId_Implementation, return FString(););
+	virtual FString GetHardwareId_Implementation() const = 0;
+
+	/** Returns whether it's valid to call `SetHardwareId()` on this device at this time. */
+	UFUNCTION(BlueprintCallable, BlueprintNativeEvent, Category="Live Link Device|Connection")
+	bool CanSetHardwareId();
+	virtual bool CanSetHardwareId_Implementation() { return false; }
 
 	/** Set hardware identifier (serial number, network endpoint, etc). */
 	UFUNCTION(BlueprintCallable, BlueprintNativeEvent, Category="Live Link Device|Connection")
 	bool SetHardwareId(const FString& HardwareID);
-	virtual bool SetHardwareId_Implementation(const FString& HardwareID) PURE_VIRTUAL(ILiveLinkDeviceCapability_Connection::SetHardwareId_Implementation, return false;);
+	virtual bool SetHardwareId_Implementation(const FString& HardwareID) { return false; }
 
+	/** Attempt to establish a connection. */
 	UFUNCTION(BlueprintCallable, BlueprintNativeEvent, Category="Live Link Device|Connection")
 	bool Connect();
-	virtual bool Connect_Implementation() PURE_VIRTUAL(ILiveLinkDeviceCapability_Connection::Connect_Implementation, return false;);
+	virtual bool Connect_Implementation() = 0;
 
+	/** Attempt to terminate an existing connection. */
 	UFUNCTION(BlueprintCallable, BlueprintNativeEvent, Category="Live Link Device|Connection")
 	bool Disconnect();
-	virtual bool Disconnect_Implementation() PURE_VIRTUAL(ILiveLinkDeviceCapability_Connection::Disconnect_Implementation, return false;);
+	virtual bool Disconnect_Implementation() = 0;
 
 	UFUNCTION(BlueprintCallable, BlueprintNativeEvent, Category="Live Link Device|Connection")
 	UConnectionDelegate* GetConnectionDelegate();
-
 	UConnectionDelegate* GetConnectionDelegate_Implementation();
 
 protected:
-
 	void SetConnectionStatus(ELiveLinkDeviceConnectionStatus InStatus);
 
 private:
-
 	TStrongObjectPtr<UConnectionDelegate> ConnectionDelegate;
 };
diff --git a/Engine/Plugins/Animation/LiveLinkDevice/Source/LiveLinkDevice/Public/LiveLinkDeviceCapability_Recording.h b/Engine/Plugins/Animation/LiveLinkDevice/Source/LiveLinkDevice/Public/LiveLinkDeviceCapability_Recording.h
--- a/Engine/Plugins/Animation/LiveLinkDevice/Source/LiveLinkDevice/Public/LiveLinkDeviceCapability_Recording.h
+++ b/Engine/Plugins/Animation/LiveLinkDevice/Source/LiveLinkDevice/Public/LiveLinkDeviceCapability_Recording.h
@@ -22,8 +22,15 @@ public:
 	virtual void OnDeviceSubsystemInitialized() override;
 	virtual void OnDeviceSubsystemDeinitializing() override;
 
-	virtual SHeaderRow::FColumn::FArguments& GenerateHeaderForColumn(const FName InColumnId, SHeaderRow::FColumn::FArguments& InArgs) const override;
-	virtual TSharedPtr<SWidget> GenerateWidgetForColumn(const FName InColumnId, ULiveLinkDevice* InDevice) const override;
+	virtual SHeaderRow::FColumn::FArguments& GenerateHeaderForColumn(
+		const FName InColumnId,
+		SHeaderRow::FColumn::FArguments& InArgs
+	) override;
+	virtual TSharedPtr<SWidget> GenerateWidgetForColumn(
+		const FName InColumnId,
+		const FLiveLinkDeviceWidgetArguments& InArgs,
+		ULiveLinkDevice* InDevice
+	) override;
 
 protected:
 	void HandleRecordingStarted();
@@ -38,13 +45,13 @@ class LIVELINKDEVICE_API ILiveLinkDeviceCapability_Recording : public ILiveLinkD
 public:
 	UFUNCTION(BlueprintCallable, BlueprintNativeEvent, Category="Live Link Device|Recording")
 	bool StartRecording();
-	virtual bool StartRecording_Implementation() PURE_VIRTUAL(ILiveLinkDeviceCapability_Recording::StartRecording_Implementation, return false;);
+	virtual bool StartRecording_Implementation() = 0;
 
 	UFUNCTION(BlueprintCallable, BlueprintNativeEvent, Category="Live Link Device|Recording")
 	bool StopRecording();
-	virtual bool StopRecording_Implementation() PURE_VIRTUAL(ILiveLinkDeviceCapability_Recording::StopRecording_Implementation, return false;);
+	virtual bool StopRecording_Implementation() = 0;
 
 	UFUNCTION(BlueprintCallable, BlueprintNativeEvent, Category="Live Link Device|Recording")
 	bool IsRecording() const;
-	virtual bool IsRecording_Implementation() const PURE_VIRTUAL(ILiveLinkDeviceCapability_Recording::IsRecording_Implementation, return false;);
+	virtual bool IsRecording_Implementation() const = 0;
 };
diff --git a/Engine/Plugins/Animation/LiveLinkDevice/Source/LiveLinkDevice/Public/LiveLinkDeviceSubsystem.h b/Engine/Plugins/Animation/LiveLinkDevice/Source/LiveLinkDevice/Public/LiveLinkDeviceSubsystem.h
--- a/Engine/Plugins/Animation/LiveLinkDevice/Source/LiveLinkDevice/Public/LiveLinkDeviceSubsystem.h
+++ b/Engine/Plugins/Animation/LiveLinkDevice/Source/LiveLinkDevice/Public/LiveLinkDeviceSubsystem.h
@@ -2,37 +2,73 @@
 
 #pragma once
 
+#include "LiveLinkHubSessionExtraData.h"
+#include "Misc/Optional.h"
 #include "Subsystems/EngineSubsystem.h"
 #include "Templates/SubclassOf.h"
+#include "Templates/ValueOrError.h"
 #include "LiveLinkDeviceSubsystem.generated.h"
 
 
 class ULiveLinkDevice;
+class ULiveLinkDeviceSettings;
 class ULiveLinkDeviceCapability;
 
 
+USTRUCT(BlueprintType)
+struct FLiveLinkDeviceCreateResult
+{
+	GENERATED_BODY()
+
+	UPROPERTY(BlueprintReadWrite, Category="Live Link|Devices")
+	FGuid DeviceId;
+
+	UPROPERTY(BlueprintReadWrite, Category="Live Link|Devices")
+	TObjectPtr<ULiveLinkDevice> Device;
+};
+
+
+USTRUCT(BlueprintType)
+struct FLiveLinkDeviceCreateError
+{
+	GENERATED_BODY()
+};
+
+
 /**
  * Device repository with lifecycle notifications.
  * Facilitates cached lookups related to device and capability classes.
  */
 UCLASS(BlueprintType)
 class LIVELINKDEVICE_API ULiveLinkDeviceSubsystem : public UEngineSubsystem
+	, public ILiveLinkHubSessionExtraDataHandler
 {
 	GENERATED_BODY()
 
 public:
 	DECLARE_MULTICAST_DELEGATE_TwoParams(FOnLiveLinkDeviceChangedDelegate, FGuid, ULiveLinkDevice*);
 
+	using FCreateResult = TValueOrError<FLiveLinkDeviceCreateResult, FLiveLinkDeviceCreateError>;
+
 	//~ Begin USubsystem interface
 	virtual void Initialize(FSubsystemCollectionBase& Collection) override;
 	virtual void Deinitialize() override;
 	//~ End USubsystem interface
 
+	//~ Begin ILiveLinkHubSessionExtraDataHandler interface
+	virtual TSubclassOf<ULiveLinkHubSessionExtraData> GetExtraDataClass() const override;
+	virtual void OnExtraDataSessionSaving(ULiveLinkHubSessionExtraData* InExtraData) override;
+	virtual void OnExtraDataSessionLoaded(const ULiveLinkHubSessionExtraData* InExtraData) override;
+	//~ End ILiveLinkHubSessionExtraDataHandler interface
+
 	/**
 	 * Register a newly created \c ULiveLinkDevice.
 	 * @return A new GUID serving as a handle to this device.
 	 */
-	FGuid AddDevice(ULiveLinkDevice* InDevice);
+	FCreateResult CreateDeviceOfClass(
+		TSubclassOf<ULiveLinkDevice> InDeviceClass,
+		ULiveLinkDeviceSettings* InDeviceSettingsTemplate = nullptr
+	);
 
 	/** Remove a previously registered device. */
 	void RemoveDevice(ULiveLinkDevice* InDevice);
@@ -68,6 +104,15 @@ public:
 	void GetDevicesByCapability(TSubclassOf<ULiveLinkDeviceCapability> Capability, TArray<ULiveLinkDevice*>& OutDevices) const;
 
 private:
+	FCreateResult InternalCreateDeviceOfClass(
+		TSubclassOf<ULiveLinkDevice> InDeviceClass,
+		ULiveLinkDeviceSettings* InDeviceSettingsTemplate = nullptr,
+		TOptional<FGuid> InDeviceGuid = {}
+	);
+
+	void InternalAddDevice(FGuid InDeviceId, ULiveLinkDevice* InDevice, ULiveLinkDeviceSettings* InSettings);
+	void InternalRemoveDevice(FGuid InDeviceId, ULiveLinkDevice* InDevice);
+
 	void RegisterCapabilityClass(TSubclassOf<ULiveLinkDeviceCapability> InCapability);
 	void RegisterDeviceClass(TSubclassOf<ULiveLinkDevice> InDeviceClass);
 
diff --git a/Engine/Plugins/Animation/LiveLinkExampleDevice/Source/LiveLinkExampleDevice/LiveLinkExampleDevice.Build.cs b/Engine/Plugins/Animation/LiveLinkExampleDevice/Source/LiveLinkExampleDevice/LiveLinkExampleDevice.Build.cs
--- a/Engine/Plugins/Animation/LiveLinkExampleDevice/Source/LiveLinkExampleDevice/LiveLinkExampleDevice.Build.cs
+++ b/Engine/Plugins/Animation/LiveLinkExampleDevice/Source/LiveLinkExampleDevice/LiveLinkExampleDevice.Build.cs
@@ -10,11 +10,11 @@ public class LiveLinkExampleDevice : ModuleRules
 			new string[] {
 				"Core",
 				"CoreUObject",
+				"Engine",
 				"InputCore",
 				"LiveLinkDevice",
 				"LiveLinkHub",
 				"Networking",
-				"OSC",
 				"Projects",
 				"Slate",
 				"SlateCore",
diff --git a/Engine/Plugins/Animation/LiveLinkExampleDevice/Source/LiveLinkExampleDevice/Private/Devices/LiveLinkExampleDevice.cpp b/Engine/Plugins/Animation/LiveLinkExampleDevice/Source/LiveLinkExampleDevice/Private/Devices/LiveLinkExampleDevice.cpp
new file mode 100644
--- /dev/null
+++ b/Engine/Plugins/Animation/LiveLinkExampleDevice/Source/LiveLinkExampleDevice/Private/Devices/LiveLinkExampleDevice.cpp
@@ -0,0 +1,210 @@
+// Copyright Epic Games, Inc. All Rights Reserved.
+
+#include "LiveLinkExampleDevice.h"
+#include "Engine/Engine.h"
+#include "ILiveLinkRecordingSessionInfo.h"
+#include "Interfaces/IPv4/IPv4Endpoint.h"
+#include "LiveLinkDeviceSubsystem.h"
+#include "Logging/StructuredLog.h"
+#include "Modules/ModuleManager.h"
+
+
+IMPLEMENT_MODULE(FDefaultModuleImpl, LiveLinkExampleDevice);
+
+
+DECLARE_LOG_CATEGORY_EXTERN(LogLiveLinkExampleDevice, Verbose, All);
+
+DEFINE_LOG_CATEGORY(LogLiveLinkExampleDevice);
+
+
+#define LOCTEXT_NAMESPACE "LiveLinkExampleDevice"
+
+
+void ULiveLinkExampleDevice::OnDeviceAdded()
+{
+}
+
+
+void ULiveLinkExampleDevice::OnDeviceRemoved()
+{
+}
+
+
+FText ULiveLinkExampleDevice::GetDisplayName() const
+{
+	return FText::FromString(GetDeviceSettings<ULiveLinkExampleDeviceSettings>()->DisplayName);
+}
+
+EDeviceHealth ULiveLinkExampleDevice::GetDeviceHealth() const
+{
+	return EDeviceHealth::Nominal;
+}
+
+FText ULiveLinkExampleDevice::GetHealthText() const
+{
+	return FText::FromString("Example Health");
+}
+
+
+void ULiveLinkExampleDevice::OnSettingChanged(const FPropertyChangedEvent& InPropertyChangedEvent)
+{
+	Super::OnSettingChanged(InPropertyChangedEvent);
+
+	const ULiveLinkExampleDeviceSettings* DeviceSettings = GetDeviceSettings<ULiveLinkExampleDeviceSettings>();
+
+	static const FName AddressName = GET_MEMBER_NAME_CHECKED(ULiveLinkExampleDeviceSettings, IpAddress);
+	static const FName PortName = GET_MEMBER_NAME_CHECKED(ULiveLinkExampleDeviceSettings, Port);
+
+	const FName ChangedPropertyName = InPropertyChangedEvent.GetPropertyName();
+	if (ChangedPropertyName == AddressName)
+	{
+		FIPv4Address Address;
+		if (FIPv4Address::Parse(DeviceSettings->IpAddress, Address))
+		{
+			UE_LOGFMT(LogLiveLinkExampleDevice, Verbose, "Device '{Name}': Changing address to {Address}",
+				DeviceSettings->DisplayName, Address.ToString());
+		}
+		else
+		{
+			UE_LOGFMT(LogLiveLinkExampleDevice, Warning, "Device '{Name}': Failed to parse address {Address}",
+				DeviceSettings->DisplayName, DeviceSettings->IpAddress);
+		}
+	}
+	else if (ChangedPropertyName == PortName)
+	{
+		UE_LOGFMT(LogLiveLinkExampleDevice, Verbose, "Device '{Name}': Changing port to {Port}",
+			DeviceSettings->DisplayName, DeviceSettings->Port);
+
+	}
+}
+
+
+ELiveLinkDeviceConnectionStatus ULiveLinkExampleDevice::GetConnectionStatus_Implementation() const
+{
+	return ELiveLinkDeviceConnectionStatus::Disconnected;
+}
+
+
+FString ULiveLinkExampleDevice::GetHardwareId_Implementation() const
+{
+	return GetDeviceSettings<ULiveLinkExampleDeviceSettings>()->IpAddress;
+}
+
+
+bool ULiveLinkExampleDevice::SetHardwareId_Implementation(const FString& InHardwareID)
+{
+	ULiveLinkExampleDeviceSettings* DeviceSettings = GetDeviceSettings<ULiveLinkExampleDeviceSettings>();
+
+	bool bParsedEndpoint = false;
+	FIPv4Endpoint Endpoint;
+	if (FIPv4Endpoint::Parse(InHardwareID, Endpoint))
+	{
+		bParsedEndpoint = true;
+	}
+	else if (FIPv4Address::Parse(InHardwareID, Endpoint.Address))
+	{
+		const UClass* SettingsClass = ULiveLinkExampleDeviceSettings::StaticClass();
+		const ULiveLinkExampleDeviceSettings* SettingsCDO =
+			SettingsClass->GetDefaultObject<ULiveLinkExampleDeviceSettings>();
+		const uint16 DefaultPort = SettingsCDO->Port;
+
+		Endpoint.Port = DefaultPort;
+
+		bParsedEndpoint = true;
+	}
+
+	if (bParsedEndpoint)
+	{
+		UE_LOGFMT(LogLiveLinkExampleDevice, Verbose, "Device '{Name}': Changing endpoint to {Endpoint}",
+			DeviceSettings->DisplayName, Endpoint.ToString());
+
+		DeviceSettings->IpAddress = Endpoint.Address.ToString();
+		DeviceSettings->Port = Endpoint.Port;
+
+		return true;
+	}
+	else
+	{
+		UE_LOGFMT(LogLiveLinkExampleDevice, Warning, "Device '{Name}': Failed to parse endpoint {HardwareID}",
+			DeviceSettings->DisplayName, InHardwareID);
+
+		return false;
+	}
+}
+
+
+bool ULiveLinkExampleDevice::Connect_Implementation()
+{
+	if (!ensure(ConnectionStatus == ELiveLinkDeviceConnectionStatus::Disconnected))
+	{
+		return false;
+	}
+
+	const ULiveLinkExampleDeviceSettings* DeviceSettings = GetDeviceSettings<ULiveLinkExampleDeviceSettings>();
+
+	FIPv4Endpoint Endpoint;
+	if (!FIPv4Address::Parse(DeviceSettings->IpAddress, Endpoint.Address))
+	{
+		return false;
+	}
+
+	Endpoint.Port = DeviceSettings->Port;
+
+	ConnectionStatus = ELiveLinkDeviceConnectionStatus::Connecting;
+	SetConnectionStatus(ConnectionStatus);
+
+	return true;
+}
+
+
+bool ULiveLinkExampleDevice::Disconnect_Implementation()
+{
+	const bool bCanDisconnect = (
+		ConnectionStatus == ELiveLinkDeviceConnectionStatus::Connected ||
+		ConnectionStatus == ELiveLinkDeviceConnectionStatus::Connecting);
+
+	if (!ensure(bCanDisconnect))
+	{
+		return false;
+	}
+
+	ConnectionStatus = ELiveLinkDeviceConnectionStatus::Disconnected;
+
+	return false;
+}
+
+
+bool ULiveLinkExampleDevice::StartRecording_Implementation()
+{
+	const ULiveLinkExampleDeviceSettings* DeviceSettings = GetDeviceSettings<ULiveLinkExampleDeviceSettings>();
+	ILiveLinkRecordingSessionInfo& SessionInfo = ILiveLinkRecordingSessionInfo::Get();
+
+	UE_LOGFMT(LogLiveLinkExampleDevice, Verbose,
+		"Device '{Name}': Started recording ({Session} / {Slate} / {Take})",
+		DeviceSettings->DisplayName,
+		SessionInfo.GetSlateName(),
+		SessionInfo.GetTakeNumber(),
+		SessionInfo.GetSessionName()
+	);
+
+	bIsRecording = true;
+
+	return true;
+}
+
+
+bool ULiveLinkExampleDevice::StopRecording_Implementation()
+{
+	bIsRecording = false;
+
+	return true;
+}
+
+
+bool ULiveLinkExampleDevice::IsRecording_Implementation() const
+{
+	return bIsRecording;
+}
+
+
+#undef LOCTEXT_NAMESPACE
diff --git a/Engine/Plugins/Animation/LiveLinkExampleDevice/Source/LiveLinkExampleDevice/Private/Devices/LiveLinkExampleDevice.h b/Engine/Plugins/Animation/LiveLinkExampleDevice/Source/LiveLinkExampleDevice/Private/Devices/LiveLinkExampleDevice.h
new file mode 100644
--- /dev/null
+++ b/Engine/Plugins/Animation/LiveLinkExampleDevice/Source/LiveLinkExampleDevice/Private/Devices/LiveLinkExampleDevice.h
@@ -0,0 +1,68 @@
+// Copyright Epic Games, Inc. All Rights Reserved.
+
+#pragma once
+
+#include "LiveLinkDevice.h"
+#include "LiveLinkDeviceCapability_Connection.h"
+#include "LiveLinkDeviceCapability_Recording.h"
+#include "LiveLinkExampleDevice.generated.h"
+
+
+UCLASS()
+class ULiveLinkExampleDeviceSettings : public ULiveLinkDeviceSettings
+{
+    GENERATED_BODY()
+
+public:
+	UPROPERTY(EditAnywhere, Category="Example Device")
+	FString DisplayName = TEXT("Example Device");
+
+	UPROPERTY(EditAnywhere, Category="Example Device")
+	FString IpAddress = TEXT("127.0.0.1");
+
+	UPROPERTY(EditAnywhere, Category="Example Device")
+	uint16 Port = 12345;
+};
+
+
+UCLASS()
+class ULiveLinkExampleDevice : public ULiveLinkDevice
+	, public ILiveLinkDeviceCapability_Connection
+	, public ILiveLinkDeviceCapability_Recording
+{
+	GENERATED_BODY()
+
+public:
+	//~ Begin ULiveLinkDevice interface
+	virtual TSubclassOf<ULiveLinkDeviceSettings> GetSettingsClass() const override
+	{
+		return ULiveLinkExampleDeviceSettings::StaticClass();
+	}
+
+	virtual FText GetDisplayName() const override;
+	virtual EDeviceHealth GetDeviceHealth() const override;
+	virtual FText GetHealthText() const override;
+
+	virtual void OnDeviceAdded() override;
+	virtual void OnDeviceRemoved() override;
+	virtual void OnSettingChanged(const FPropertyChangedEvent& InPropertyChangedEvent) override;
+	//~ End ULiveLinkDevice interface
+
+	//~ Begin ILiveLinkDeviceCapability_Connection interface
+	virtual ELiveLinkDeviceConnectionStatus GetConnectionStatus_Implementation() const override;
+	virtual FString GetHardwareId_Implementation() const override;
+	virtual bool SetHardwareId_Implementation(const FString& HardwareID) override;
+	virtual bool Connect_Implementation() override;
+	virtual bool Disconnect_Implementation() override;
+	//~ End ILiveLinkDeviceCapability_Connection interface
+
+	//~ Begin ILiveLinkDeviceCapability_Recording interface
+	virtual bool StartRecording_Implementation() override;
+	virtual bool StopRecording_Implementation() override;
+	virtual bool IsRecording_Implementation() const override;
+	//~ End ILiveLinkDeviceCapability_Recording interface
+
+private:
+	ELiveLinkDeviceConnectionStatus ConnectionStatus = ELiveLinkDeviceConnectionStatus::Disconnected;
+	bool bIsRecording = false;
+};
diff --git a/Engine/Plugins/Animation/LiveLinkExampleDevice/Source/LiveLinkExampleDevice/Private/Devices/LiveLinkUnrealDevice.cpp b/Engine/Plugins/Animation/LiveLinkExampleDevice/Source/LiveLinkExampleDevice/Private/Devices/LiveLinkUnrealDevice.cpp
deleted file mode 100644
diff --git a/Engine/Plugins/Animation/LiveLinkExampleDevice/Source/LiveLinkExampleDevice/Private/Devices/LiveLinkUnrealDevice.h b/Engine/Plugins/Animation/LiveLinkExampleDevice/Source/LiveLinkExampleDevice/Private/Devices/LiveLinkUnrealDevice.h
deleted file mode 100644
diff --git a/Engine/Plugins/Animation/LiveLinkExampleDevice/Source/LiveLinkExampleDevice/Private/Tests/LiveLinkExampleDeviceTests.cpp b/Engine/Plugins/Animation/LiveLinkExampleDevice/Source/LiveLinkExampleDevice/Private/Tests/LiveLinkExampleDeviceTests.cpp
new file mode 100644
--- /dev/null
+++ b/Engine/Plugins/Animation/LiveLinkExampleDevice/Source/LiveLinkExampleDevice/Private/Tests/LiveLinkExampleDeviceTests.cpp
@@ -0,0 +1,32 @@
+// Copyright Epic Games, Inc. All Rights Reserved.
+
+#include "Devices/LiveLinkExampleDevice.h"
+#include "Engine/Engine.h"
+#include "LiveLinkDeviceSubsystem.h"
+#include "Misc/AutomationTest.h"
+
+
+IMPLEMENT_SIMPLE_AUTOMATION_TEST(FExampleDeviceTest, "LiveLinkHub.ExampleDevice",
+	EAutomationTestFlags::EditorContext | EAutomationTestFlags::EngineFilter)
+
+bool FExampleDeviceTest::RunTest(const FString& Parameters)
+{
+	ULiveLinkDeviceSubsystem* DeviceSubsystem = GEngine->GetEngineSubsystem<ULiveLinkDeviceSubsystem>();
+
+	ULiveLinkExampleDeviceSettings* DeviceSettingsTemplate = NewObject<ULiveLinkExampleDeviceSettings>();
+	DeviceSettingsTemplate->IpAddress = FString(TEXT("127.1.2.3"));
+	DeviceSettingsTemplate->DisplayName = FString(TEXT("Test Device"));
+
+	ULiveLinkDeviceSubsystem::FCreateResult CreateResult =
+		DeviceSubsystem->CreateDeviceOfClass(ULiveLinkExampleDevice::StaticClass(), DeviceSettingsTemplate);
+
+	if (!TestTrueExpr(CreateResult.HasValue()))
+	{
+		return false;
+	}
+
+	const FGuid DeviceId = CreateResult.GetValue().DeviceId;
+	ULiveLinkDevice* NewDevice = CreateResult.GetValue().Device;
+
+	return true;
+}
diff --git a/Engine/Plugins/Animation/LiveLinkExampleDevice/Source/LiveLinkExampleDevice/Private/Tests/LiveLinkUnrealDeviceTests.cpp b/Engine/Plugins/Animation/LiveLinkExampleDevice/Source/LiveLinkExampleDevice/Private/Tests/LiveLinkUnrealDeviceTests.cpp
deleted file mode 100644
diff --git a/Engine/Plugins/VirtualProduction/CaptureManager/CaptureManagerApp/Source/CaptureManagerEditor/Private/Scripting/ULiveLinkHubCaptureDevice.cpp b/Engine/Plugins/VirtualProduction/CaptureManager/CaptureManagerApp/Source/CaptureManagerEditor/Private/Scripting/ULiveLinkHubCaptureDevice.cpp
--- a/Engine/Plugins/VirtualProduction/CaptureManager/CaptureManagerApp/Source/CaptureManagerEditor/Private/Scripting/ULiveLinkHubCaptureDevice.cpp
+++ b/Engine/Plugins/VirtualProduction/CaptureManager/CaptureManagerApp/Source/CaptureManagerEditor/Private/Scripting/ULiveLinkHubCaptureDevice.cpp
@@ -294,14 +294,11 @@ ULiveLinkHubCaptureDevice* ULiveLinkHubCaptureDeviceFactory::CreateDeviceByClass
 	ULiveLinkDeviceSubsystem* DeviceSubsystem = GEngine->GetEngineSubsystem<ULiveLinkDeviceSubsystem>();
 	check(DeviceSubsystem);
 
-	TObjectPtr<ULiveLinkDevice> NewDevice = NewObject<ULiveLinkDevice>(DeviceSubsystem, DeviceClass);
+	ULiveLinkDeviceSubsystem::FCreateResult DeviceResult = DeviceSubsystem->CreateDeviceOfClass(DeviceClass);
+	ULiveLinkDevice* NewDevice = DeviceResult.GetValue().Device;
 
 	check(NewDevice->Implements<ULiveLinkDeviceCapability_Ingest>());
 
-	const FGuid DeviceGuid = DeviceSubsystem->AddDevice(NewDevice);
-
-	NewDevice->OnDeviceAdded(DeviceGuid, InSettings);
-
 	TObjectPtr<ULiveLinkHubCaptureDevice> Device = NewObject<ULiveLinkHubCaptureDevice>(this);
 	Device->Name = MoveTemp(InName);
 	Device->ImplPtr->BackgroundDevice = MoveTemp(NewDevice);
diff --git a/Engine/Plugins/VirtualProduction/CaptureManager/CaptureManagerApp/Source/IngestLiveLinkDevice/Private/BaseIngestLiveLinkDevice.cpp b/Engine/Plugins/VirtualProduction/CaptureManager/CaptureManagerApp/Source/IngestLiveLinkDevice/Private/BaseIngestLiveLinkDevice.cpp
--- a/Engine/Plugins/VirtualProduction/CaptureManager/CaptureManagerApp/Source/IngestLiveLinkDevice/Private/BaseIngestLiveLinkDevice.cpp
+++ b/Engine/Plugins/VirtualProduction/CaptureManager/CaptureManagerApp/Source/IngestLiveLinkDevice/Private/BaseIngestLiveLinkDevice.cpp
@@ -123,7 +123,7 @@ void UBaseIngestLiveLinkDevice::FImpl::IngestTake(int32 InTakeId, const UIngestC
 	{
 		if (!Uploader.IsValid())
 		{
-			Uploader = CommunicationManager->GetUploader(Owner->DeviceGuid,
+			Uploader = CommunicationManager->GetUploader(Owner->GetDeviceId(),
 														 Owner->GetDisplayName().ToString(),
 														 FUploader::FUploadStateUpdate::CreateRaw(this, &FImpl::UploadTakeStateHandler),
 														 FUploader::FUploadFinished::CreateRaw(this, &FImpl::UploadTakeFinished));
@@ -316,19 +316,9 @@ void UBaseIngestLiveLinkDevice::FImpl::TryCleanUpWorkingDirectory(const FString&
 	}
 }
 
-
-FGuid UBaseIngestLiveLinkDevice::GetDeviceGuid() const
-{
-	return DeviceGuid;
-}
-
-void UBaseIngestLiveLinkDevice::OnDeviceAdded(const FGuid InDeviceGuid, ULiveLinkDeviceSettings* InSettings)
+void UBaseIngestLiveLinkDevice::OnDeviceAdded()
 {
-	DeviceGuid = InDeviceGuid;
-
 	Impl = MakePimpl<FImpl>(this, UE::CaptureManager::Private::GetCommunicationManager());
-
-	Super::OnDeviceAdded(DeviceGuid, InSettings);
 }
 
 void UBaseIngestLiveLinkDevice::OnDeviceRemoved()
diff --git a/Engine/Plugins/VirtualProduction/CaptureManager/CaptureManagerApp/Source/IngestLiveLinkDevice/Public/BaseIngestLiveLinkDevice.h b/Engine/Plugins/VirtualProduction/CaptureManager/CaptureManagerApp/Source/IngestLiveLinkDevice/Public/BaseIngestLiveLinkDevice.h
--- a/Engine/Plugins/VirtualProduction/CaptureManager/CaptureManagerApp/Source/IngestLiveLinkDevice/Public/BaseIngestLiveLinkDevice.h
+++ b/Engine/Plugins/VirtualProduction/CaptureManager/CaptureManagerApp/Source/IngestLiveLinkDevice/Public/BaseIngestLiveLinkDevice.h
@@ -22,21 +22,13 @@ class INGESTLIVELINKDEVICE_API UBaseIngestLiveLinkDevice : public ULiveLinkDevic
 
 public:
 
-	/**
-	 * Function provides the Device Guid defined by the LiveLink Device Subsystem during device
-	 * @returns Device unique ID
-	 */
-	FGuid GetDeviceGuid() const;
-
 	//~ Begin ULiveLinkDevice interface
 
 	/**
 	 * Inherited function from the ULiveLinkDevice that is called when device is added
 	 * @note Users are required to invoke this function if it is inherited in the users code.
-	 * @params DeviceGuid Device unique ID
-	 * @params Settigns Device settings
 	 */
-	virtual void OnDeviceAdded(const FGuid InDeviceGuid, ULiveLinkDeviceSettings* InSettings) override;
+	virtual void OnDeviceAdded() override;
 
 	/**
 	 * Inherited function from the ULiveLinkDevice that is called when device is removed
@@ -85,6 +77,4 @@ private:
 
 	class FImpl;
 	TPimplPtr<FImpl> Impl;
-
-	FGuid DeviceGuid;
 };
diff --git a/Engine/Plugins/VirtualProduction/CaptureManager/CaptureManagerApp/Source/LiveLinkCapabilities/Private/Ingest/LiveLinkDeviceCapability_Ingest.cpp b/Engine/Plugins/VirtualProduction/CaptureManager/CaptureManagerApp/Source/LiveLinkCapabilities/Private/Ingest/LiveLinkDeviceCapability_Ingest.cpp
--- a/Engine/Plugins/VirtualProduction/CaptureManager/CaptureManagerApp/Source/LiveLinkCapabilities/Private/Ingest/LiveLinkDeviceCapability_Ingest.cpp
+++ b/Engine/Plugins/VirtualProduction/CaptureManager/CaptureManagerApp/Source/LiveLinkCapabilities/Private/Ingest/LiveLinkDeviceCapability_Ingest.cpp
@@ -5,7 +5,7 @@
 #include "LiveLinkDevice.h"
 #include "Widgets/Text/STextBlock.h"
 
-SHeaderRow::FColumn::FArguments& ULiveLinkDeviceCapability_Ingest::GenerateHeaderForColumn(const FName InColumnId, SHeaderRow::FColumn::FArguments& InArgs) const
+SHeaderRow::FColumn::FArguments& ULiveLinkDeviceCapability_Ingest::GenerateHeaderForColumn(const FName InColumnId, SHeaderRow::FColumn::FArguments& InArgs)
 {
 	if (InColumnId == Column_IngestSupport)
 	{
@@ -17,7 +17,7 @@ SHeaderRow::FColumn::FArguments& ULiveLinkDeviceCapability_Ingest::GenerateHeade
 	return Super::GenerateHeaderForColumn(InColumnId, InArgs);
 }
 
-TSharedPtr<SWidget> ULiveLinkDeviceCapability_Ingest::GenerateWidgetForColumn(const FName InColumnId, ULiveLinkDevice* InDevice) const
+TSharedPtr<SWidget> ULiveLinkDeviceCapability_Ingest::GenerateWidgetForColumn(const FName InColumnId, const FLiveLinkDeviceWidgetArguments& InArgs, ULiveLinkDevice* InDevice)
 {
 	return SNew(STextBlock)
 		.Text_Lambda(
diff --git a/Engine/Plugins/VirtualProduction/CaptureManager/CaptureManagerApp/Source/LiveLinkCapabilities/Public/Ingest/LiveLinkDeviceCapability_Ingest.h b/Engine/Plugins/VirtualProduction/CaptureManager/CaptureManagerApp/Source/LiveLinkCapabilities/Public/Ingest/LiveLinkDeviceCapability_Ingest.h
--- a/Engine/Plugins/VirtualProduction/CaptureManager/CaptureManagerApp/Source/LiveLinkCapabilities/Public/Ingest/LiveLinkDeviceCapability_Ingest.h
+++ b/Engine/Plugins/VirtualProduction/CaptureManager/CaptureManagerApp/Source/LiveLinkCapabilities/Public/Ingest/LiveLinkDeviceCapability_Ingest.h
@@ -29,8 +29,8 @@ public:
 	{
 	}
 
-	virtual SHeaderRow::FColumn::FArguments& GenerateHeaderForColumn(const FName InColumnId, SHeaderRow::FColumn::FArguments& InArgs) const override;
-	virtual TSharedPtr<SWidget> GenerateWidgetForColumn(const FName InColumnId, ULiveLinkDevice* InDevice) const override;
+	virtual SHeaderRow::FColumn::FArguments& GenerateHeaderForColumn(const FName InColumnId, SHeaderRow::FColumn::FArguments& InArgs) override;
+	virtual TSharedPtr<SWidget> GenerateWidgetForColumn(const FName InColumnId, const FLiveLinkDeviceWidgetArguments& InArgs, ULiveLinkDevice* InDevice) override;
 };
 
 class LIVELINKCAPABILITIES_API ILiveLinkDeviceCapability_Ingest 
diff --git a/Engine/Plugins/VirtualProduction/CaptureManager/CaptureManagerDevices/Source/CPSLiveLinkDevice/Private/LiveLinkFaceDevice.cpp b/Engine/Plugins/VirtualProduction/CaptureManager/CaptureManagerDevices/Source/CPSLiveLinkDevice/Private/LiveLinkFaceDevice.cpp
--- a/Engine/Plugins/VirtualProduction/CaptureManager/CaptureManagerDevices/Source/CPSLiveLinkDevice/Private/LiveLinkFaceDevice.cpp
+++ b/Engine/Plugins/VirtualProduction/CaptureManager/CaptureManagerDevices/Source/CPSLiveLinkDevice/Private/LiveLinkFaceDevice.cpp
@@ -19,10 +19,10 @@
 
 const ULiveLinkFaceDeviceSettings* ULiveLinkFaceDevice::GetSettings() const
 {
-	return CastChecked<ULiveLinkFaceDeviceSettings>(Settings);
+	return GetDeviceSettings<ULiveLinkFaceDeviceSettings>();
 }
 
-TSubclassOf<ULiveLinkDeviceSettings> ULiveLinkFaceDevice::StaticGetSettingsClass() const
+TSubclassOf<ULiveLinkDeviceSettings> ULiveLinkFaceDevice::GetSettingsClass() const
 {
 	return ULiveLinkFaceDeviceSettings::StaticClass();
 }
@@ -42,11 +42,9 @@ FText ULiveLinkFaceDevice::GetHealthText() const
 	return FText::FromString("Example Health");
 }
 
-void ULiveLinkFaceDevice::OnDeviceAdded(const FGuid InDeviceGuid, ULiveLinkDeviceSettings* InSettings)
+void ULiveLinkFaceDevice::OnDeviceAdded()
 {
-	CastChecked<ULiveLinkFaceDeviceSettings>(InSettings)->ConnectAction.DeviceGuid = InDeviceGuid;
-
-	Super::OnDeviceAdded(InDeviceGuid, InSettings);
+	GetDeviceSettings<ULiveLinkFaceDeviceSettings>()->ConnectAction.DeviceGuid = GetDeviceId();
 }
 
 void ULiveLinkFaceDevice::OnDeviceRemoved()
diff --git a/Engine/Plugins/VirtualProduction/CaptureManager/CaptureManagerDevices/Source/CPSLiveLinkDevice/Private/LiveLinkFaceDevice.h b/Engine/Plugins/VirtualProduction/CaptureManager/CaptureManagerDevices/Source/CPSLiveLinkDevice/Private/LiveLinkFaceDevice.h
--- a/Engine/Plugins/VirtualProduction/CaptureManager/CaptureManagerDevices/Source/CPSLiveLinkDevice/Private/LiveLinkFaceDevice.h
+++ b/Engine/Plugins/VirtualProduction/CaptureManager/CaptureManagerDevices/Source/CPSLiveLinkDevice/Private/LiveLinkFaceDevice.h
@@ -52,12 +52,12 @@ public:
 	const ULiveLinkFaceDeviceSettings* GetSettings() const;
 
 	//~ Begin ULiveLinkDevice interface
-	virtual TSubclassOf<ULiveLinkDeviceSettings> StaticGetSettingsClass() const override;
+	virtual TSubclassOf<ULiveLinkDeviceSettings> GetSettingsClass() const override;
 	virtual FText GetDisplayName() const override;
 	virtual EDeviceHealth GetDeviceHealth() const override;
 	virtual FText GetHealthText() const override;
 
-	virtual void OnDeviceAdded(const FGuid InDeviceGuid, ULiveLinkDeviceSettings* InSettings) override;
+	virtual void OnDeviceAdded() override;
 	virtual void OnDeviceRemoved() override;
 	//~ End ULiveLinkDevice interface
 
diff --git a/Engine/Plugins/VirtualProduction/CaptureManager/CaptureManagerDevices/Source/FilesystemLiveLinkDevice/Private/FilesystemLiveLinkDevice.cpp b/Engine/Plugins/VirtualProduction/CaptureManager/CaptureManagerDevices/Source/FilesystemLiveLinkDevice/Private/FilesystemLiveLinkDevice.cpp
--- a/Engine/Plugins/VirtualProduction/CaptureManager/CaptureManagerDevices/Source/FilesystemLiveLinkDevice/Private/FilesystemLiveLinkDevice.cpp
+++ b/Engine/Plugins/VirtualProduction/CaptureManager/CaptureManagerDevices/Source/FilesystemLiveLinkDevice/Private/FilesystemLiveLinkDevice.cpp
@@ -16,10 +16,10 @@ DEFINE_LOG_CATEGORY(LogGenericLiveLinkArchiveDevice);
 
 const UFilesystemLiveLinkDeviceSettings* UFilesystemLiveLinkDevice::GetSettings() const
 {
-	return CastChecked<UFilesystemLiveLinkDeviceSettings>(Settings);
+	return GetDeviceSettings<UFilesystemLiveLinkDeviceSettings>();
 }
 
-TSubclassOf<ULiveLinkDeviceSettings> UFilesystemLiveLinkDevice::StaticGetSettingsClass() const
+TSubclassOf<ULiveLinkDeviceSettings> UFilesystemLiveLinkDevice::GetSettingsClass() const
 {
 	return UFilesystemLiveLinkDeviceSettings::StaticClass();
 }
diff --git a/Engine/Plugins/VirtualProduction/CaptureManager/CaptureManagerDevices/Source/FilesystemLiveLinkDevice/Private/FilesystemLiveLinkDevice.h b/Engine/Plugins/VirtualProduction/CaptureManager/CaptureManagerDevices/Source/FilesystemLiveLinkDevice/Private/FilesystemLiveLinkDevice.h
--- a/Engine/Plugins/VirtualProduction/CaptureManager/CaptureManagerDevices/Source/FilesystemLiveLinkDevice/Private/FilesystemLiveLinkDevice.h
+++ b/Engine/Plugins/VirtualProduction/CaptureManager/CaptureManagerDevices/Source/FilesystemLiveLinkDevice/Private/FilesystemLiveLinkDevice.h
@@ -41,7 +41,7 @@ public:
 	const UFilesystemLiveLinkDeviceSettings* GetSettings() const;
 
 	//~ Begin ULiveLinkDevice interface
-	virtual TSubclassOf<ULiveLinkDeviceSettings> StaticGetSettingsClass() const override;
+	virtual TSubclassOf<ULiveLinkDeviceSettings> GetSettingsClass() const override;
 	virtual FText GetDisplayName() const override;
 	virtual EDeviceHealth GetDeviceHealth() const override;
 	virtual FText GetHealthText() const override;
diff --git a/Engine/Plugins/VirtualProduction/CaptureManager/CaptureManagerDevices/Source/MonoVideoLiveLinkDevice/Private/MonoVideoLiveLinkDevice.cpp b/Engine/Plugins/VirtualProduction/CaptureManager/CaptureManagerDevices/Source/MonoVideoLiveLinkDevice/Private/MonoVideoLiveLinkDevice.cpp
--- a/Engine/Plugins/VirtualProduction/CaptureManager/CaptureManagerDevices/Source/MonoVideoLiveLinkDevice/Private/MonoVideoLiveLinkDevice.cpp
+++ b/Engine/Plugins/VirtualProduction/CaptureManager/CaptureManagerDevices/Source/MonoVideoLiveLinkDevice/Private/MonoVideoLiveLinkDevice.cpp
@@ -30,10 +30,10 @@ FMediaRWManager& GetMediaRWManager()
 
 const UMonoVideoLiveLinkDeviceSettings* UMonoVideoLiveLinkDevice::GetSettings() const
 {
-	return CastChecked<UMonoVideoLiveLinkDeviceSettings>(Settings);
+	return GetDeviceSettings<UMonoVideoLiveLinkDeviceSettings>();
 }
 
-TSubclassOf<ULiveLinkDeviceSettings> UMonoVideoLiveLinkDevice::StaticGetSettingsClass() const
+TSubclassOf<ULiveLinkDeviceSettings> UMonoVideoLiveLinkDevice::GetSettingsClass() const
 {
 	return UMonoVideoLiveLinkDeviceSettings::StaticClass();
 }
diff --git a/Engine/Plugins/VirtualProduction/CaptureManager/CaptureManagerDevices/Source/MonoVideoLiveLinkDevice/Private/MonoVideoLiveLinkDevice.h b/Engine/Plugins/VirtualProduction/CaptureManager/CaptureManagerDevices/Source/MonoVideoLiveLinkDevice/Private/MonoVideoLiveLinkDevice.h
--- a/Engine/Plugins/VirtualProduction/CaptureManager/CaptureManagerDevices/Source/MonoVideoLiveLinkDevice/Private/MonoVideoLiveLinkDevice.h
+++ b/Engine/Plugins/VirtualProduction/CaptureManager/CaptureManagerDevices/Source/MonoVideoLiveLinkDevice/Private/MonoVideoLiveLinkDevice.h
@@ -57,7 +57,7 @@ public:
 	const UMonoVideoLiveLinkDeviceSettings* GetSettings() const;
 
 	//~ Begin ULiveLinkDevice interface
-	virtual TSubclassOf<ULiveLinkDeviceSettings> StaticGetSettingsClass() const override;
+	virtual TSubclassOf<ULiveLinkDeviceSettings> GetSettingsClass() const override;
 	virtual FText GetDisplayName() const override;
 	virtual EDeviceHealth GetDeviceHealth() const override;
 	virtual FText GetHealthText() const override;
