Add missing entry in LiveLinkDeviceStyle

diff --git a/Engine/Plugins/Animation/LiveLinkDevice/Source/LiveLinkDevice/Private/LiveLinkDeviceStyle.cpp b/Engine/Plugins/Animation/LiveLinkDevice/Source/LiveLinkDevice/Private/LiveLinkDeviceStyle.cpp
--- a/Engine/Plugins/Animation/LiveLinkDevice/Source/LiveLinkDevice/Private/LiveLinkDeviceStyle.cpp
+++ b/Engine/Plugins/Animation/LiveLinkDevice/Source/LiveLinkDevice/Private/LiveLinkDeviceStyle.cpp
@@ -27,6 +27,8 @@ void FLiveLinkDeviceStyle::Initialize()
 	StyleSet->Set("Record", new FSlateVectorImageBrush(PluginContentRoot / "PlayControlsRecord.svg", Icon16x16));
 	StyleSet->Set("Record.Monochrome", new FSlateVectorImageBrush(PluginContentRoot / "PlayControlsRecord_Monochrome.svg", Icon16x16));
 
+	StyleSet->Set("LiveLinkHub.Devices.Icon", new FSlateVectorImageBrush(PluginContentRoot / "Devices.svg", Icon16x16));
+
 	FSlateStyleRegistry::RegisterSlateStyle(*StyleSet);
 }
 
