[MH-16397] Resolving crash when broadcasting a dynamic delegate from mutliple threads

diff --git a/Engine/Plugins/Animation/LiveLinkDevice/Source/LiveLinkDevice/Private/LiveLinkDeviceCapability_Connection.cpp b/Engine/Plugins/Animation/LiveLinkDevice/Source/LiveLinkDevice/Private/LiveLinkDeviceCapability_Connection.cpp
--- a/Engine/Plugins/Animation/LiveLinkDevice/Source/LiveLinkDevice/Private/LiveLinkDeviceCapability_Connection.cpp
+++ b/Engine/Plugins/Animation/LiveLinkDevice/Source/LiveLinkDevice/Private/LiveLinkDeviceCapability_Connection.cpp
@@ -186,7 +186,21 @@ UConnectionDelegate* ILiveLinkDeviceCapability_Connection::GetConnectionDelegate
 void ILiveLinkDeviceCapability_Connection::SetConnectionStatus(ELiveLinkDeviceConnectionStatus InStatus)
 {
 	ConnectionDelegate->ConnectionChanged.Broadcast(InStatus);
-	ConnectionDelegate->ConnectionChangedDynamic.Broadcast(InStatus);
+
+	if (IsInGameThread())
+	{
+		ConnectionDelegate->ConnectionChangedDynamic.Broadcast(InStatus);
+	}
+	else
+	{
+		ExecuteOnGameThread(TEXT("InvokeConnectionDelegate"), [Delegate = TWeakObjectPtr<UConnectionDelegate>(ConnectionDelegate.Get()), Status = InStatus]()
+		{
+			if (TStrongObjectPtr<UConnectionDelegate> StrongDelegate = Delegate.Pin())
+			{
+				StrongDelegate->ConnectionChangedDynamic.Broadcast(Status);
+			}
+		});
+	}
 }
 
 
diff --git a/Engine/Plugins/Animation/LiveLinkDevice/Source/LiveLinkDevice/Public/LiveLinkDeviceCapability_Connection.h b/Engine/Plugins/Animation/LiveLinkDevice/Source/LiveLinkDevice/Public/LiveLinkDeviceCapability_Connection.h
--- a/Engine/Plugins/Animation/LiveLinkDevice/Source/LiveLinkDevice/Public/LiveLinkDeviceCapability_Connection.h
+++ b/Engine/Plugins/Animation/LiveLinkDevice/Source/LiveLinkDevice/Public/LiveLinkDeviceCapability_Connection.h
@@ -68,12 +68,14 @@ class LIVELINKDEVICE_API UConnectionDelegate : public UObject
 
 public:
 
+	/** Dynamic delegate to be used in Blueprints. Avoid using it outside of the Game Thread. */
 	UPROPERTY(EditAnywhere, Category = "Live Link Device|Connection")
 	FDeviceConnectionStatusChangedDynamic ConnectionChangedDynamic;
 
 	FDeviceConnectionStatusChanged ConnectionChanged;
 };
 
+
 /**
  * Status and operations relevant to devices with the concept of being (dis)connected.
  * This could be a physical hardware connection, a network connection, or something else.
