LiveLink - Remove LiveLinkHub plugin dependency from LiveLinkDevice to allow LiveLink to depend on LiveLinkDevice

#Synced-CL 48684930

diff --git a/Engine/Plugins/Animation/LiveLinkDevice/LiveLinkDevice.uplugin b/Engine/Plugins/Animation/LiveLinkDevice/LiveLinkDevice.uplugin
--- a/Engine/Plugins/Animation/LiveLinkDevice/LiveLinkDevice.uplugin
+++ b/Engine/Plugins/Animation/LiveLinkDevice/LiveLinkDevice.uplugin
@@ -16,12 +16,6 @@
 	"IsExperimentalVersion": true,
 	"Hidden": true,
 	"Installed": false,
-	"Plugins": [
-		{
-			"Name": "LiveLinkHub",
-			"Enabled": true
-		}
-	],
 	"Modules": [
 		{
 			"Name": "LiveLinkDevice",
diff --git a/Engine/Plugins/Animation/LiveLinkDevice/Source/LiveLinkDevice/LiveLinkDevice.Build.cs b/Engine/Plugins/Animation/LiveLinkDevice/Source/LiveLinkDevice/LiveLinkDevice.Build.cs
--- a/Engine/Plugins/Animation/LiveLinkDevice/Source/LiveLinkDevice/LiveLinkDevice.Build.cs
+++ b/Engine/Plugins/Animation/LiveLinkDevice/Source/LiveLinkDevice/LiveLinkDevice.Build.cs
@@ -6,6 +6,12 @@ public class LiveLinkDevice : ModuleRules
 {
 	public LiveLinkDevice(ReadOnlyTargetRules Target) : base(Target)
 	{
+		PublicDependencyModuleNames.AddRange(
+			new string[] {
+				"LiveLinkInterface",
+			}
+		);
+
 		PrivateDependencyModuleNames.AddRange(
 			new string[] {
 				"Core",
@@ -13,7 +19,6 @@ public class LiveLinkDevice : ModuleRules
 				"Engine",
 				"InputCore",
 				"JsonUtilities",
-				"LiveLinkHub",
 				"Projects",
 				"Slate",
 				"SlateCore",
diff --git a/Engine/Plugins/Animation/LiveLinkDevice/Source/LiveLinkDevice/Private/LiveLinkDeviceCapability_Recording.cpp b/Engine/Plugins/Animation/LiveLinkDevice/Source/LiveLinkDevice/Private/LiveLinkDeviceCapability_Recording.cpp
--- a/Engine/Plugins/Animation/LiveLinkDevice/Source/LiveLinkDevice/Private/LiveLinkDeviceCapability_Recording.cpp
+++ b/Engine/Plugins/Animation/LiveLinkDevice/Source/LiveLinkDevice/Private/LiveLinkDeviceCapability_Recording.cpp
@@ -2,10 +2,10 @@
 
 #include "LiveLinkDeviceCapability_Recording.h"
 #include "Engine/Engine.h"
-#include "ILiveLinkRecordingSessionInfo.h"
 #include "LiveLinkDevice.h"
 #include "LiveLinkDeviceStyle.h"
 #include "LiveLinkDeviceSubsystem.h"
+#include "LiveLinkHub/ILiveLinkRecordingSession.h"
 #include "Widgets/Images/SImage.h"
 
 
@@ -50,13 +50,13 @@ protected:
 void ULiveLinkDeviceCapability_Recording::OnDeviceSubsystemInitialized()
 {
 	IModularFeatures& ModularFeatures = IModularFeatures::Get();
-	if (ModularFeatures.IsModularFeatureAvailable(ILiveLinkRecordingSessionInfo::GetModularFeatureName()))
+	if (ModularFeatures.IsModularFeatureAvailable(ILiveLinkRecordingSession::GetModularFeatureName()))
 	{
-		ILiveLinkRecordingSessionInfo::Get().OnRecordingStarted().AddUObject(
+		ILiveLinkRecordingSession::Get().OnRecordingStarted().AddUObject(
 			this, &ULiveLinkDeviceCapability_Recording::HandleRecordingStarted
 		);
 
-		ILiveLinkRecordingSessionInfo::Get().OnRecordingStopped().AddUObject(
+		ILiveLinkRecordingSession::Get().OnRecordingStopped().AddUObject(
 			this, &ULiveLinkDeviceCapability_Recording::HandleRecordingStopped
 		);
 	}
@@ -66,10 +66,10 @@ void ULiveLinkDeviceCapability_Recording::OnDeviceSubsystemInitialized()
 void ULiveLinkDeviceCapability_Recording::OnDeviceSubsystemDeinitializing()
 {
 	IModularFeatures& ModularFeatures = IModularFeatures::Get();
-	if (ModularFeatures.IsModularFeatureAvailable(ILiveLinkRecordingSessionInfo::GetModularFeatureName()))
+	if (ModularFeatures.IsModularFeatureAvailable(ILiveLinkRecordingSession::GetModularFeatureName()))
 	{
-		ILiveLinkRecordingSessionInfo::Get().OnRecordingStarted().RemoveAll(this);
-		ILiveLinkRecordingSessionInfo::Get().OnRecordingStopped().RemoveAll(this);
+		ILiveLinkRecordingSession::Get().OnRecordingStarted().RemoveAll(this);
+		ILiveLinkRecordingSession::Get().OnRecordingStopped().RemoveAll(this);
 	}
 }
 
diff --git a/Engine/Plugins/Animation/LiveLinkDevice/Source/LiveLinkDevice/Private/LiveLinkHubSessionExtraData_Device.h b/Engine/Plugins/Animation/LiveLinkDevice/Source/LiveLinkDevice/Private/LiveLinkHubSessionExtraData_Device.h
--- a/Engine/Plugins/Animation/LiveLinkDevice/Source/LiveLinkDevice/Private/LiveLinkHubSessionExtraData_Device.h
+++ b/Engine/Plugins/Animation/LiveLinkDevice/Source/LiveLinkDevice/Private/LiveLinkHubSessionExtraData_Device.h
@@ -2,7 +2,7 @@
 
 #pragma once
 
-#include "LiveLinkHubSessionExtraData.h"
+#include "LiveLinkHub/LiveLinkHubSessionExtraData.h"
 #include "Templates/SubclassOf.h"
 #include "LiveLinkHubSessionExtraData_Device.generated.h"
 
diff --git a/Engine/Plugins/Animation/LiveLinkDevice/Source/LiveLinkDevice/Public/LiveLinkDeviceSubsystem.h b/Engine/Plugins/Animation/LiveLinkDevice/Source/LiveLinkDevice/Public/LiveLinkDeviceSubsystem.h
--- a/Engine/Plugins/Animation/LiveLinkDevice/Source/LiveLinkDevice/Public/LiveLinkDeviceSubsystem.h
+++ b/Engine/Plugins/Animation/LiveLinkDevice/Source/LiveLinkDevice/Public/LiveLinkDeviceSubsystem.h
@@ -2,7 +2,8 @@
 
 #pragma once
 
-#include "LiveLinkHubSessionExtraData.h"
+#include "LiveLinkHub/LiveLinkHubSessionExtraData.h"
+
 #include "Misc/Optional.h"
 #include "Subsystems/EngineSubsystem.h"
 #include "Templates/SubclassOf.h"
@@ -14,6 +15,7 @@ class ULiveLinkDevice;
 class ULiveLinkDeviceSettings;
 class ULiveLinkDeviceCapability;
 
+#define UE_API LIVELINKDEVICE_API
 
 USTRUCT(BlueprintType)
 struct FLiveLinkDeviceCreateResult
@@ -39,8 +41,8 @@ struct FLiveLinkDeviceCreateError
  * Device repository with lifecycle notifications.
  * Facilitates cached lookups related to device and capability classes.
  */
-UCLASS(BlueprintType)
-class LIVELINKDEVICE_API ULiveLinkDeviceSubsystem : public UEngineSubsystem
+UCLASS(MinimalAPI, BlueprintType)
+class ULiveLinkDeviceSubsystem : public UEngineSubsystem
 	, public ILiveLinkHubSessionExtraDataHandler
 {
 	GENERATED_BODY()
@@ -51,30 +53,30 @@ public:
 	using FCreateResult = TValueOrError<FLiveLinkDeviceCreateResult, FLiveLinkDeviceCreateError>;
 
 	//~ Begin USubsystem interface
-	virtual void Initialize(FSubsystemCollectionBase& Collection) override;
-	virtual void Deinitialize() override;
+	UE_API virtual void Initialize(FSubsystemCollectionBase& Collection) override;
+	UE_API virtual void Deinitialize() override;
 	//~ End USubsystem interface
 
 	//~ Begin ILiveLinkHubSessionExtraDataHandler interface
-	virtual TSubclassOf<ULiveLinkHubSessionExtraData> GetExtraDataClass() const override;
-	virtual void OnExtraDataSessionSaving(ULiveLinkHubSessionExtraData* InExtraData) override;
-	virtual void OnExtraDataSessionLoaded(const ULiveLinkHubSessionExtraData* InExtraData) override;
+	UE_API virtual TSubclassOf<ULiveLinkHubSessionExtraData> GetExtraDataClass() const override;
+	UE_API virtual void OnExtraDataSessionSaving(ULiveLinkHubSessionExtraData* InExtraData) override;
+	UE_API virtual void OnExtraDataSessionLoaded(const ULiveLinkHubSessionExtraData* InExtraData) override;
 	//~ End ILiveLinkHubSessionExtraDataHandler interface
 
 	/**
 	 * Register a newly created \c ULiveLinkDevice.
 	 * @return A new GUID serving as a handle to this device.
 	 */
-	FCreateResult CreateDeviceOfClass(
+	UE_API FCreateResult CreateDeviceOfClass(
 		TSubclassOf<ULiveLinkDevice> InDeviceClass,
 		ULiveLinkDeviceSettings* InDeviceSettingsTemplate = nullptr
 	);
 
 	/** Remove a previously registered device. */
-	void RemoveDevice(ULiveLinkDevice* InDevice);
+	UE_API void RemoveDevice(ULiveLinkDevice* InDevice);
 
 	/** Remove a previously registered device by its GUID. */
-	void RemoveDevice(FGuid InDeviceId);
+	UE_API void RemoveDevice(FGuid InDeviceId);
 
 	FOnLiveLinkDeviceChangedDelegate& OnDeviceAdded() { return OnDeviceAddedDelegate; }
 	FOnLiveLinkDeviceChangedDelegate& OnDeviceRemoved() { return OnDeviceRemovedDelegate; }
@@ -98,10 +100,10 @@ public:
 	const TMap<FName, TSubclassOf<ULiveLinkDeviceCapability>>& GetTableColumnIdToCapability() { return TableColumnIdToCapability; }
 
 	UFUNCTION(BlueprintPure, Category="Live Link|Devices", meta=(DeterminesOutputType="DeviceClass"))
-	void GetDevicesByClass(TSubclassOf<ULiveLinkDevice> DeviceClass, TArray<ULiveLinkDevice*>& OutDevices) const;
+	UE_API void GetDevicesByClass(TSubclassOf<ULiveLinkDevice> DeviceClass, TArray<ULiveLinkDevice*>& OutDevices) const;
 
 	UFUNCTION(BlueprintPure, Category="Live Link|Devices")
-	void GetDevicesByCapability(TSubclassOf<ULiveLinkDeviceCapability> Capability, TArray<ULiveLinkDevice*>& OutDevices) const;
+	UE_API void GetDevicesByCapability(TSubclassOf<ULiveLinkDeviceCapability> Capability, TArray<ULiveLinkDevice*>& OutDevices) const;
 
 private:
 	FCreateResult InternalCreateDeviceOfClass(
@@ -136,3 +138,5 @@ private:
 
 	TMap<FName, TSubclassOf<ULiveLinkDeviceCapability>> TableColumnIdToCapability;
 };
+
+#undef UE_API
\ No newline at end of file
diff --git a/Engine/Plugins/Animation/LiveLinkExampleDevice/Source/LiveLinkExampleDevice/Private/Devices/LiveLinkExampleDevice.cpp b/Engine/Plugins/Animation/LiveLinkExampleDevice/Source/LiveLinkExampleDevice/Private/Devices/LiveLinkExampleDevice.cpp
--- a/Engine/Plugins/Animation/LiveLinkExampleDevice/Source/LiveLinkExampleDevice/Private/Devices/LiveLinkExampleDevice.cpp
+++ b/Engine/Plugins/Animation/LiveLinkExampleDevice/Source/LiveLinkExampleDevice/Private/Devices/LiveLinkExampleDevice.cpp
@@ -2,9 +2,9 @@
 
 #include "LiveLinkExampleDevice.h"
 #include "Engine/Engine.h"
-#include "ILiveLinkRecordingSessionInfo.h"
 #include "Interfaces/IPv4/IPv4Endpoint.h"
 #include "LiveLinkDeviceSubsystem.h"
+#include "LiveLinkHub/ILiveLinkRecordingSession.h"
 #include "Logging/StructuredLog.h"
 #include "Modules/ModuleManager.h"
 
@@ -177,7 +177,7 @@ bool ULiveLinkExampleDevice::Disconnect_Implementation()
 bool ULiveLinkExampleDevice::StartRecording_Implementation()
 {
 	const ULiveLinkExampleDeviceSettings* DeviceSettings = GetDeviceSettings<ULiveLinkExampleDeviceSettings>();
-	ILiveLinkRecordingSessionInfo& SessionInfo = ILiveLinkRecordingSessionInfo::Get();
+	ILiveLinkRecordingSession& SessionInfo = ILiveLinkRecordingSession::Get();
 
 	UE_LOGFMT(LogLiveLinkExampleDevice, Verbose,
 		"Device '{Name}': Started recording ({Session} / {Slate} / {Take})",
diff --git a/Engine/Plugins/Animation/LiveLinkHub/Config/DefaultLiveLinkHub.ini b/Engine/Plugins/Animation/LiveLinkHub/Config/DefaultLiveLinkHub.ini
new file mode 100644
--- /dev/null
+++ b/Engine/Plugins/Animation/LiveLinkHub/Config/DefaultLiveLinkHub.ini
@@ -0,0 +1,2 @@
+[CoreRedirects]
++ClassRedirects=(OldName="/Script/LiveLinkHub.LiveLinkHubSessionExtraData",NewName="/Script/LiveLinkInterface.LiveLinkHubSessionExtraData")
\ No newline at end of file
diff --git a/Engine/Plugins/Animation/LiveLinkHub/Source/LiveLinkHub/LiveLinkHub.Build.cs b/Engine/Plugins/Animation/LiveLinkHub/Source/LiveLinkHub/LiveLinkHub.Build.cs
--- a/Engine/Plugins/Animation/LiveLinkHub/Source/LiveLinkHub/LiveLinkHub.Build.cs
+++ b/Engine/Plugins/Animation/LiveLinkHub/Source/LiveLinkHub/LiveLinkHub.Build.cs
@@ -12,6 +12,7 @@ public class LiveLinkHub : ModuleRules
 			{
 				"Core",
 				"CoreUObject",
+				"LiveLinkInterface" // Can be removed in 5.10. This dependency was added because some headers were moved to LiveLinkInterface module.
 			});
 
 		PrivateDependencyModuleNames.AddRange(
diff --git a/Engine/Plugins/Animation/LiveLinkHub/Source/LiveLinkHub/Private/Config/LiveLinkHubTemplateTokens.cpp b/Engine/Plugins/Animation/LiveLinkHub/Source/LiveLinkHub/Private/Config/LiveLinkHubTemplateTokens.cpp
--- a/Engine/Plugins/Animation/LiveLinkHub/Source/LiveLinkHub/Private/Config/LiveLinkHubTemplateTokens.cpp
+++ b/Engine/Plugins/Animation/LiveLinkHub/Source/LiveLinkHub/Private/Config/LiveLinkHubTemplateTokens.cpp
@@ -3,6 +3,7 @@
 #include "LiveLinkHubTemplateTokens.h"
 
 #include "ILiveLinkHubModule.h"
+#include "LiveLinkHub/ILiveLinkRecordingSession.h"
 #include "Session/LiveLinkHubSessionManager.h"
 
 #include UE_INLINE_GENERATED_CPP_BY_NAME(LiveLinkHubTemplateTokens)
@@ -66,7 +67,7 @@ void ULiveLinkHubNamingTokens::OnPreEvaluate_Implementation(const FNamingTokensE
 			{
 				ConfigName = FPaths::GetBaseFilename(SessionManager->GetLastConfigPath());
 				
-				const ILiveLinkRecordingSessionInfo& SessionInfo = ILiveLinkRecordingSessionInfo::Get();
+				const ILiveLinkRecordingSession& SessionInfo = ILiveLinkRecordingSession::Get();
 				SessionName = SessionInfo.GetSessionName();
 				TakeNumber = SessionInfo.GetTakeNumber();
 				SlateName = SessionInfo.GetSlateName();
diff --git a/Engine/Plugins/Animation/LiveLinkHub/Source/LiveLinkHub/Private/Recording/Implementations/LiveLinkUAssetRecorder.cpp b/Engine/Plugins/Animation/LiveLinkHub/Source/LiveLinkHub/Private/Recording/Implementations/LiveLinkUAssetRecorder.cpp
--- a/Engine/Plugins/Animation/LiveLinkHub/Source/LiveLinkHub/Private/Recording/Implementations/LiveLinkUAssetRecorder.cpp
+++ b/Engine/Plugins/Animation/LiveLinkHub/Source/LiveLinkHub/Private/Recording/Implementations/LiveLinkUAssetRecorder.cpp
@@ -13,11 +13,11 @@
 #include "Framework/MultiBox/MultiBoxBuilder.h"
 #include "HAL/PlatformFileManager.h"
 #include "IContentBrowserSingleton.h"
-#include "ILiveLinkRecordingSessionInfo.h"
 #include "LiveLinkHub.h"
 #include "LiveLinkHubClient.h"
 #include "LiveLinkHubLog.h"
 #include "LiveLinkHubModule.h"
+#include "LiveLinkHub/ILiveLinkRecordingSession.h"
 #include "LiveLinkPreset.h"
 #include "LiveLinkTypes.h"
 #include "Misc/FileHelper.h"
@@ -114,7 +114,7 @@ bool FLiveLinkUAssetRecorder::CanRecord(FText* OutErrorMessage) const
 		return false;
 	}
 	
-	const ILiveLinkRecordingSessionInfo& SessionInfo = ILiveLinkRecordingSessionInfo::Get();
+	const ILiveLinkRecordingSession& SessionInfo = ILiveLinkRecordingSession::Get();
 	if (SessionInfo.GetSessionName().IsEmpty())
 	{
 		ErrorText = LOCTEXT("SessionEmptyTooltip", "Session cannot be empty");
@@ -368,7 +368,7 @@ void FLiveLinkUAssetRecorder::SaveRecording(FLiveLinkUAssetRecordingData&& InRec
 	// Increment take on successful recording.
 	if (bIncrementTake)
 	{
-		ILiveLinkRecordingSessionInfo& SessionInfo = ILiveLinkRecordingSessionInfo::Get();
+		ILiveLinkRecordingSession& SessionInfo = ILiveLinkRecordingSession::Get();
 		SessionInfo.SetTakeNumber(SessionInfo.GetTakeNumber() + 1);
 	}
 }
diff --git a/Engine/Plugins/Animation/LiveLinkHub/Source/LiveLinkHub/Private/Recording/LiveLinkHubRecordingController.cpp b/Engine/Plugins/Animation/LiveLinkHub/Source/LiveLinkHub/Private/Recording/LiveLinkHubRecordingController.cpp
--- a/Engine/Plugins/Animation/LiveLinkHub/Source/LiveLinkHub/Private/Recording/LiveLinkHubRecordingController.cpp
+++ b/Engine/Plugins/Animation/LiveLinkHub/Source/LiveLinkHub/Private/Recording/LiveLinkHubRecordingController.cpp
@@ -2,8 +2,8 @@
 
 #include "LiveLinkHubRecordingController.h"
 
-#include "ILiveLinkRecordingSessionInfo.h"
 #include "Implementations/LiveLinkUAssetRecorder.h"
+#include "LiveLinkHub/ILiveLinkRecordingSession.h"
 #include "LiveLinkRecorder.h"
 #include "SLiveLinkHubRecordingView.h"
 #include "Widgets/DeclarativeSyntaxSupport.h"
@@ -29,20 +29,20 @@ TSharedRef<SWidget> FLiveLinkHubRecordingController::MakeRecordToolbarEntry()
 
 void FLiveLinkHubRecordingController::StartRecording()
 {
-	ILiveLinkRecordingSessionInfo::Get().OnRecordingStarted().Broadcast();
+	ILiveLinkRecordingSession::Get().OnRecordingStarted().Broadcast();
 
 	RecorderImplementation->StartRecording();
 }
 	
 void FLiveLinkHubRecordingController::StopRecording()
 {
-	ILiveLinkRecordingSessionInfo::Get().OnRecordingStopped().Broadcast();
+	ILiveLinkRecordingSession::Get().OnRecordingStopped().Broadcast();
 
 	RecorderImplementation->StopRecording();
 
 	if (GetDefault<ULiveLinkHubSettings>()->bIncrementTakeOnRecordingEnd)
 	{
-		ILiveLinkRecordingSessionInfo& SessionInfo = ILiveLinkRecordingSessionInfo::Get();
+		ILiveLinkRecordingSession& SessionInfo = ILiveLinkRecordingSession::Get();
 		SessionInfo.SetTakeNumber(SessionInfo.GetTakeNumber() + 1);
 	}
 }
diff --git a/Engine/Plugins/Animation/LiveLinkHub/Source/LiveLinkHub/Private/Session/LiveLinkHubSessionData.h b/Engine/Plugins/Animation/LiveLinkHub/Source/LiveLinkHub/Private/Session/LiveLinkHubSessionData.h
--- a/Engine/Plugins/Animation/LiveLinkHub/Source/LiveLinkHub/Private/Session/LiveLinkHubSessionData.h
+++ b/Engine/Plugins/Animation/LiveLinkHub/Source/LiveLinkHub/Private/Session/LiveLinkHubSessionData.h
@@ -5,8 +5,8 @@
 #include "Clients/LiveLinkHubUEClientInfo.h"
 #include "CoreTypes.h"
 #include "LiveLinkHubMessages.h"
+#include "LiveLinkHub/LiveLinkHubSessionExtraData.h"
 #include "LiveLinkPresetTypes.h"
-#include "LiveLinkHubSessionExtraData.h"
 #include "Misc/Guid.h"
 
 #include "LiveLinkHubSessionData.generated.h"
diff --git a/Engine/Plugins/Animation/LiveLinkHub/Source/LiveLinkHub/Private/Session/LiveLinkHubSessionManager.h b/Engine/Plugins/Animation/LiveLinkHub/Source/LiveLinkHub/Private/Session/LiveLinkHubSessionManager.h
--- a/Engine/Plugins/Animation/LiveLinkHub/Source/LiveLinkHub/Private/Session/LiveLinkHubSessionManager.h
+++ b/Engine/Plugins/Animation/LiveLinkHub/Source/LiveLinkHub/Private/Session/LiveLinkHubSessionManager.h
@@ -10,10 +10,10 @@
 #include "Engine/Engine.h"
 #include "HAL/CriticalSection.h"
 #include "IDesktopPlatform.h"
-#include "ILiveLinkRecordingSessionInfo.h"
 #include "LiveLinkHubClient.h"
 #include "LiveLinkHubLog.h"
-#include "LiveLinkHubSessionExtraData.h"
+#include "LiveLinkHub/ILiveLinkRecordingSession.h"
+#include "LiveLinkHub/LiveLinkHubSessionExtraData.h"
 #include "LiveLinkSourceSettings.h"
 #include "Misc/ConfigCacheIni.h"
 #include "Recording/LiveLinkHubRecordingController.h"
@@ -79,12 +79,12 @@ public:
 };
 
 class FLiveLinkHubSessionManager : public ILiveLinkHubSessionManager
-	, public ILiveLinkRecordingSessionInfo
+	, public ILiveLinkRecordingSession
 {
 public:
 	FLiveLinkHubSessionManager()
 	{
-		IModularFeatures::Get().RegisterModularFeature(ILiveLinkRecordingSessionInfo::GetModularFeatureName(), this);
+		IModularFeatures::Get().RegisterModularFeature(ILiveLinkRecordingSession::GetModularFeatureName(), this);
 
 		FScopeLock Lock(&CurrentSessionCS);
 		CurrentSession = MakeShared<FLiveLinkHubSession>(OnClientAddedToSessionDelegate, OnClientRemovedFromSessionDelegate);
@@ -92,7 +92,7 @@ public:
 
 	virtual ~FLiveLinkHubSessionManager() override
 	{
-		IModularFeatures::Get().UnregisterModularFeature(ILiveLinkRecordingSessionInfo::GetModularFeatureName(), this);
+		IModularFeatures::Get().UnregisterModularFeature(ILiveLinkRecordingSession::GetModularFeatureName(), this);
 	}
 
 	//~ Begin LiveLinkHubSessionManager
diff --git a/Engine/Plugins/Animation/LiveLinkHub/Source/LiveLinkHub/Private/UI/Widgets/SLiveLinkRecordingSessionInfo.cpp b/Engine/Plugins/Animation/LiveLinkHub/Source/LiveLinkHub/Private/UI/Widgets/SLiveLinkRecordingSessionInfo.cpp
--- a/Engine/Plugins/Animation/LiveLinkHub/Source/LiveLinkHub/Private/UI/Widgets/SLiveLinkRecordingSessionInfo.cpp
+++ b/Engine/Plugins/Animation/LiveLinkHub/Source/LiveLinkHub/Private/UI/Widgets/SLiveLinkRecordingSessionInfo.cpp
@@ -2,7 +2,7 @@
 
 #include "SLiveLinkRecordingSessionInfo.h"
 
-#include "ILiveLinkRecordingSessionInfo.h"
+#include "LiveLinkHub/ILiveLinkRecordingSession.h"
 #include "SPositiveActionButton.h"
 #include "Widgets/Input/SEditableTextBox.h"
 #include "Widgets/Text/STextBlock.h"
@@ -16,7 +16,7 @@ void SLiveLinkRecordingSessionInfo::Construct(const FArguments& InArgs)
 {
 	const FMargin IntraElementPadding = FMargin(6.0, 4.0);
 
-	ILiveLinkRecordingSessionInfo& SessionInfo = ILiveLinkRecordingSessionInfo::Get();
+	ILiveLinkRecordingSession& SessionInfo = ILiveLinkRecordingSession::Get();
 
 	ChildSlot
 	[
diff --git a/Engine/Plugins/Animation/LiveLinkHub/Source/LiveLinkHub/Public/ILiveLinkRecordingSession.h b/Engine/Plugins/Animation/LiveLinkHub/Source/LiveLinkHub/Public/ILiveLinkRecordingSession.h
--- a/Engine/Plugins/Animation/LiveLinkHub/Source/LiveLinkHub/Public/ILiveLinkRecordingSession.h
+++ b/Engine/Plugins/Animation/LiveLinkHub/Source/LiveLinkHub/Public/ILiveLinkRecordingSession.h
@@ -2,52 +2,6 @@
 
 #pragma once
 
-#include "Containers/StringFwd.h"
-#include "Delegates/Delegate.h"
-#include "Features/IModularFeatures.h"
+UE_DEPRECATED_HEADER(5.8, "Include LiveLinkHub/ILiveLinkRecordingSession.h from the LiveLinkInterface module instead.")
 
-
-/** Provides an interface for querying and controlling the recording session. */
-class ILiveLinkRecordingSession : public IModularFeature
-{
-public:
-	DECLARE_MULTICAST_DELEGATE_OneParam(FOnSessionStringChanged, FStringView);
-	DECLARE_MULTICAST_DELEGATE_OneParam(FOnSessionIntChanged, int32);
-
-	/** Name of the modular feature. */
-	static FName GetModularFeatureName()
-	{
-		static const FName ModularFeatureName("LiveLinkRecordingSession");
-		return ModularFeatureName;
-	}
-
-	/** Convenience accessor for modular implementation. */
-	static ILiveLinkRecordingSession& Get()
-	{
-		IModularFeatures& ModularFeatures = IModularFeatures::Get();
-		check(ModularFeatures.IsModularFeatureAvailable(GetModularFeatureName()));
-		ensure(ModularFeatures.GetModularFeatureImplementationCount(GetModularFeatureName()) == 1);
-		return ModularFeatures.GetModularFeature<ILiveLinkRecordingSession>(GetModularFeatureName());
-	}
-
-	virtual FString GetSessionName() const = 0;
-	virtual FString GetSlateName() const = 0;
-	virtual int32 GetTakeNumber() const = 0;
-
-	virtual bool SetSessionName(FStringView InSessionName) = 0;
-	virtual bool SetSlateName(FStringView InSlateName) = 0;
-	virtual bool SetTakeNumber(int32 InTakeNumber) = 0;
-
-	virtual FOnSessionStringChanged& OnSessionNameChanged() = 0;
-	virtual FOnSessionStringChanged& OnSlateNameChanged() = 0;
-	virtual FOnSessionIntChanged& OnTakeNumberChanged() = 0;
-
-	virtual bool IsRecording() const = 0;
-
-	virtual FSimpleMulticastDelegate& OnRecordingStarted() = 0;
-	virtual FSimpleMulticastDelegate& OnRecordingStopped() = 0;
-
-	virtual bool CanRecord() const = 0;
-	virtual bool StartRecording() = 0;
-	virtual bool StopRecording() = 0;
-};
+#include "LiveLinkHub/ILiveLinkRecordingSession.h"
diff --git a/Engine/Plugins/Animation/LiveLinkHub/Source/LiveLinkHub/Public/ILiveLinkRecordingSessionInfo.h b/Engine/Plugins/Animation/LiveLinkHub/Source/LiveLinkHub/Public/ILiveLinkRecordingSessionInfo.h
--- a/Engine/Plugins/Animation/LiveLinkHub/Source/LiveLinkHub/Public/ILiveLinkRecordingSessionInfo.h
+++ b/Engine/Plugins/Animation/LiveLinkHub/Source/LiveLinkHub/Public/ILiveLinkRecordingSessionInfo.h
@@ -2,7 +2,10 @@
 
 #pragma once
 
-#include "ILiveLinkRecordingSession.h"
+UE_DEPRECATED_HEADER(5.8, "Include LiveLinkHub/ILiveLinkRecordingSession.h from the LiveLinkInterface module instead.")
 
+#include "LiveLinkHub/ILiveLinkRecordingSession.h"
 
+UE_DEPRECATED(5.8, "Use ILiveLinkRecordingSession directly instead.")
 using ILiveLinkRecordingSessionInfo = ILiveLinkRecordingSession;
+
diff --git a/Engine/Plugins/Animation/LiveLinkHub/Source/LiveLinkHub/Public/LiveLinkHubSessionExtraData.h b/Engine/Plugins/Animation/LiveLinkHub/Source/LiveLinkHub/Public/LiveLinkHubSessionExtraData.h
--- a/Engine/Plugins/Animation/LiveLinkHub/Source/LiveLinkHub/Public/LiveLinkHubSessionExtraData.h
+++ b/Engine/Plugins/Animation/LiveLinkHub/Source/LiveLinkHub/Public/LiveLinkHubSessionExtraData.h
@@ -2,56 +2,6 @@
 
 #pragma once
 
-#include "Features/IModularFeatures.h"
-#include "Templates/SubclassOf.h"
-#include "LiveLinkHubSessionExtraData.generated.h"
+UE_DEPRECATED_HEADER(5.8, "Include LiveLinkHub/LiveLinkHubSessionExtraData.h from the LiveLinkInterface module instead.")
 
-
-/** Derive from this class to define additional UProperties that participate in session serialization. */
-UCLASS(MinimalAPI, Abstract)
-class ULiveLinkHubSessionExtraData : public UObject
-{
-	GENERATED_BODY()
-};
-
-
-/** Implement this interface to provide session save/load handlers for your extra data subclass. */
-class ILiveLinkHubSessionExtraDataHandler : public IModularFeature
-{
-public:
-	/** Which derived type this handler will receive save/load events for. */
-	virtual TSubclassOf<ULiveLinkHubSessionExtraData> GetExtraDataClass() const = 0;
-
-	/** Update your extra data fields prior to session save. */
-	virtual void OnExtraDataSessionSaving(ULiveLinkHubSessionExtraData* InExtraData) = 0;
-
-	/** Handle session load. If the session was previously saved without your handler active, InExtraData may be null. */
-	virtual void OnExtraDataSessionLoaded(const ULiveLinkHubSessionExtraData* InExtraData) = 0;
-
-public:
-	/** The modular feature name used for handler registration/iteration. */
-	static FName GetModularFeatureName()
-	{
-		static const FName ModularFeatureName("LiveLinkHubSessionExtraDataHandler");
-		return ModularFeatureName;
-	}
-
-	/** Call this from your derived class to participate in session save/load. */
-	void RegisterExtraDataHandler()
-	{
-		IModularFeatures::Get().RegisterModularFeature(GetModularFeatureName(), this);
-	}
-
-	/** Call this from your derived class when finished. */
-	void UnregisterExtraDataHandler()
-	{
-		IModularFeatures::Get().UnregisterModularFeature(GetModularFeatureName(), this);
-	}
-
-	/** Used internally to iterate over registered implementations. */
-	static TArray<ILiveLinkHubSessionExtraDataHandler*> GetRegisteredHandlers()
-	{
-		return IModularFeatures::Get().GetModularFeatureImplementations<ILiveLinkHubSessionExtraDataHandler>(
-			ILiveLinkHubSessionExtraDataHandler::GetModularFeatureName());
-	}
-};
+#include "LiveLinkHub/LiveLinkHubSessionExtraData.h"
diff --git a/Engine/Plugins/Animation/LiveLinkUnrealDevice/LiveLinkUnrealDevice.uplugin b/Engine/Plugins/Animation/LiveLinkUnrealDevice/LiveLinkUnrealDevice.uplugin
--- a/Engine/Plugins/Animation/LiveLinkUnrealDevice/LiveLinkUnrealDevice.uplugin
+++ b/Engine/Plugins/Animation/LiveLinkUnrealDevice/LiveLinkUnrealDevice.uplugin
@@ -19,6 +19,10 @@
 		{
 			"Name": "LiveLinkDevice",
 			"Enabled": true
+		},
+		{
+			"Name": "LiveLinkHub",
+			"Enabled": true
 		}
 	],
 	"Modules": [
diff --git a/Engine/Plugins/Animation/LiveLinkUnrealDevice/Source/LiveLinkUnrealDevice/Private/Devices/LiveLinkUnrealDevice.cpp b/Engine/Plugins/Animation/LiveLinkUnrealDevice/Source/LiveLinkUnrealDevice/Private/Devices/LiveLinkUnrealDevice.cpp
--- a/Engine/Plugins/Animation/LiveLinkUnrealDevice/Source/LiveLinkUnrealDevice/Private/Devices/LiveLinkUnrealDevice.cpp
+++ b/Engine/Plugins/Animation/LiveLinkUnrealDevice/Source/LiveLinkUnrealDevice/Private/Devices/LiveLinkUnrealDevice.cpp
@@ -2,7 +2,7 @@
 
 #include "LiveLinkUnrealDevice.h"
 #include "Engine/Engine.h"
-#include "ILiveLinkRecordingSession.h"
+#include "LiveLinkHub/ILiveLinkRecordingSession.h"
 #include "INetworkMessagingExtension.h"
 #include "LiveLinkDeviceSubsystem.h"
 #include "LiveLinkUnrealDeviceMessages.h"
diff --git a/Engine/Plugins/VirtualProduction/CaptureManager/CaptureManagerDevices/Source/CPSLiveLinkDevice/Private/LiveLinkFaceDevice.cpp b/Engine/Plugins/VirtualProduction/CaptureManager/CaptureManagerDevices/Source/CPSLiveLinkDevice/Private/LiveLinkFaceDevice.cpp
--- a/Engine/Plugins/VirtualProduction/CaptureManager/CaptureManagerDevices/Source/CPSLiveLinkDevice/Private/LiveLinkFaceDevice.cpp
+++ b/Engine/Plugins/VirtualProduction/CaptureManager/CaptureManagerDevices/Source/CPSLiveLinkDevice/Private/LiveLinkFaceDevice.cpp
@@ -2,22 +2,22 @@
 
 #include "LiveLinkFaceDevice.h"
 
-#include "Settings/CaptureManagerSettings.h"
-
-#include "ILiveLinkRecordingSessionInfo.h"
-
 #include "Control/Messages/Constants.h"
 #include "Control/Messages/ControlUpdate.h"
 
+#include "HAL/FileManager.h"
+
+#include "LiveLinkHub/ILiveLinkRecordingSession.h"
+
 #include "Protocol/CPSDataStream.h"
 #include "Protocol/CPSFileStream.h"
 
 #include "LiveLinkFaceMetadata.h"
+#include "Settings/CaptureManagerSettings.h"
 #include "StereoCameraMetadataParseUtils.h"
 #include "Utils/CaptureExtractTimecode.h"
 #include "Utils/ParseTakeUtils.h"
 
-#include "HAL/FileManager.h"
 
 DEFINE_LOG_CATEGORY_STATIC(LogLiveLinkFaceDevice, Log, All);
 
@@ -283,7 +283,7 @@ bool ULiveLinkFaceDevice::StartRecording_Implementation()
 
 	using namespace UE::CaptureManager;
 
-	ILiveLinkRecordingSessionInfo& SessionInfo = ILiveLinkRecordingSessionInfo::Get();
+	ILiveLinkRecordingSession& SessionInfo = ILiveLinkRecordingSession::Get();
 	if (SessionInfo.GetSlateName().IsEmpty() || SessionInfo.GetTakeNumber() == -1)
 	{
 		return false;
diff --git a/Engine/Source/Runtime/LiveLinkInterface/Public/LiveLinkHub/ILiveLinkRecordingSession.h b/Engine/Source/Runtime/LiveLinkInterface/Public/LiveLinkHub/ILiveLinkRecordingSession.h
new file mode 100644
--- /dev/null
+++ b/Engine/Source/Runtime/LiveLinkInterface/Public/LiveLinkHub/ILiveLinkRecordingSession.h
@@ -0,0 +1,53 @@
+﻿// Copyright Epic Games, Inc. All Rights Reserved.
+
+#pragma once
+
+#include "Containers/StringFwd.h"
+#include "Delegates/Delegate.h"
+#include "Features/IModularFeatures.h"
+
+
+/** Provides an interface for querying and controlling the recording session. */
+class ILiveLinkRecordingSession : public IModularFeature
+{
+public:
+	DECLARE_MULTICAST_DELEGATE_OneParam(FOnSessionStringChanged, FStringView);
+	DECLARE_MULTICAST_DELEGATE_OneParam(FOnSessionIntChanged, int32);
+
+	/** Name of the modular feature. */
+	static FName GetModularFeatureName()
+	{
+		static const FName ModularFeatureName("LiveLinkRecordingSession");
+		return ModularFeatureName;
+	}
+
+	/** Convenience accessor for modular implementation. */
+	static ILiveLinkRecordingSession& Get()
+	{
+		IModularFeatures& ModularFeatures = IModularFeatures::Get();
+		check(ModularFeatures.IsModularFeatureAvailable(GetModularFeatureName()));
+		ensure(ModularFeatures.GetModularFeatureImplementationCount(GetModularFeatureName()) == 1);
+		return ModularFeatures.GetModularFeature<ILiveLinkRecordingSession>(GetModularFeatureName());
+	}
+
+	virtual FString GetSessionName() const = 0;
+	virtual FString GetSlateName() const = 0;
+	virtual int32 GetTakeNumber() const = 0;
+
+	virtual bool SetSessionName(FStringView InSessionName) = 0;
+	virtual bool SetSlateName(FStringView InSlateName) = 0;
+	virtual bool SetTakeNumber(int32 InTakeNumber) = 0;
+
+	virtual FOnSessionStringChanged& OnSessionNameChanged() = 0;
+	virtual FOnSessionStringChanged& OnSlateNameChanged() = 0;
+	virtual FOnSessionIntChanged& OnTakeNumberChanged() = 0;
+
+	virtual bool IsRecording() const = 0;
+
+	virtual FSimpleMulticastDelegate& OnRecordingStarted() = 0;
+	virtual FSimpleMulticastDelegate& OnRecordingStopped() = 0;
+
+	virtual bool CanRecord() const = 0;
+	virtual bool StartRecording() = 0;
+	virtual bool StopRecording() = 0;
+};
diff --git a/Engine/Source/Runtime/LiveLinkInterface/Public/LiveLinkHub/LiveLinkHubSessionExtraData.h b/Engine/Source/Runtime/LiveLinkInterface/Public/LiveLinkHub/LiveLinkHubSessionExtraData.h
new file mode 100644
--- /dev/null
+++ b/Engine/Source/Runtime/LiveLinkInterface/Public/LiveLinkHub/LiveLinkHubSessionExtraData.h
@@ -0,0 +1,57 @@
+// Copyright Epic Games, Inc. All Rights Reserved.
+
+#pragma once
+
+#include "Features/IModularFeatures.h"
+#include "Templates/SubclassOf.h"
+#include "LiveLinkHubSessionExtraData.generated.h"
+
+
+/** Derive from this class to define additional UProperties that participate in session serialization. */
+UCLASS(MinimalAPI, Abstract)
+class ULiveLinkHubSessionExtraData : public UObject
+{
+	GENERATED_BODY()
+};
+
+
+/** Implement this interface to provide session save/load handlers for your extra data subclass. */
+class ILiveLinkHubSessionExtraDataHandler : public IModularFeature
+{
+public:
+	/** Which derived type this handler will receive save/load events for. */
+	virtual TSubclassOf<ULiveLinkHubSessionExtraData> GetExtraDataClass() const = 0;
+
+	/** Update your extra data fields prior to session save. */
+	virtual void OnExtraDataSessionSaving(ULiveLinkHubSessionExtraData* InExtraData) = 0;
+
+	/** Handle session load. If the session was previously saved without your handler active, InExtraData may be null. */
+	virtual void OnExtraDataSessionLoaded(const ULiveLinkHubSessionExtraData* InExtraData) = 0;
+
+public:
+	/** The modular feature name used for handler registration/iteration. */
+	static FName GetModularFeatureName()
+	{
+		static const FName ModularFeatureName("LiveLinkHubSessionExtraDataHandler");
+		return ModularFeatureName;
+	}
+
+	/** Call this from your derived class to participate in session save/load. */
+	void RegisterExtraDataHandler()
+	{
+		IModularFeatures::Get().RegisterModularFeature(GetModularFeatureName(), this);
+	}
+
+	/** Call this from your derived class when finished. */
+	void UnregisterExtraDataHandler()
+	{
+		IModularFeatures::Get().UnregisterModularFeature(GetModularFeatureName(), this);
+	}
+
+	/** Used internally to iterate over registered implementations. */
+	static TArray<ILiveLinkHubSessionExtraDataHandler*> GetRegisteredHandlers()
+	{
+		return IModularFeatures::Get().GetModularFeatureImplementations<ILiveLinkHubSessionExtraDataHandler>(
+			ILiveLinkHubSessionExtraDataHandler::GetModularFeatureName());
+	}
+};
