LiveLinkDevice - Allow generating a widget for columns that aren't added dynamically

#Synced-CL 48735856

diff --git a/Engine/Plugins/Animation/LiveLinkDevice/Source/LiveLinkDevice/Private/LiveLinkDevice.cpp b/Engine/Plugins/Animation/LiveLinkDevice/Source/LiveLinkDevice/Private/LiveLinkDevice.cpp
--- a/Engine/Plugins/Animation/LiveLinkDevice/Source/LiveLinkDevice/Private/LiveLinkDevice.cpp
+++ b/Engine/Plugins/Animation/LiveLinkDevice/Source/LiveLinkDevice/Private/LiveLinkDevice.cpp
@@ -3,21 +3,31 @@
 #include "LiveLinkDevice.h"
 #include "Engine/Engine.h"
 #include "JsonObjectConverter.h"
+#include "LiveLinkDefaultColumns.h"
 #include "LiveLinkDeviceCapability.h"
 #include "LiveLinkDeviceModule.h"
 #include "LiveLinkDeviceSubsystem.h"
+#include "LiveLinkDeviceStyle.h"
 #include "Logging/StructuredLog.h"
+#include "Widgets/Images/SImage.h"
 #include "Widgets/SNullWidget.h"
 
 
-
 TSharedRef<SWidget> ULiveLinkDevice::GenerateWidgetForColumn(const FName InColumnId, const FLiveLinkDeviceWidgetArguments& InArgs)
 {
+	// Handle generating a health widget for the status column.
+	if (InColumnId == UE::LiveLink::DefaultColumn::Status)
+	{
+		return SNew(SImage)
+			.Image_UObject(this, &ULiveLinkDevice::GetHealthImage)
+			.ToolTipText_UObject(this, &ULiveLinkDevice::GetHealthText);
+	}
+
 	// Check if the capability provides a default widget implementation.
 	ULiveLinkDeviceSubsystem* Subsystem = GEngine->GetEngineSubsystem<ULiveLinkDeviceSubsystem>();
 	const TMap<FName, TSubclassOf<ULiveLinkDeviceCapability>>& ColumnIdToCapability = Subsystem->GetTableColumnIdToCapability();
 	const TSubclassOf<ULiveLinkDeviceCapability>* MaybeCapabilityClass = ColumnIdToCapability.Find(InColumnId);
-	if (ensure(MaybeCapabilityClass))
+	if (MaybeCapabilityClass)
 	{
 		ULiveLinkDeviceCapability* CapabilityCDO = MaybeCapabilityClass->GetDefaultObject();
 		if (TSharedPtr<SWidget> CapabilityWidget = CapabilityCDO->GenerateWidgetForColumn(InColumnId, InArgs, this))
@@ -27,10 +37,20 @@ TSharedRef<SWidget> ULiveLinkDevice::GenerateWidgetForColumn(const FName InColum
 	}
 
 	// Neither your device class nor the capability created a widget for this column.
-	ensure(false);
 	return SNullWidget::NullWidget;
 }
 
+const FSlateBrush* ULiveLinkDevice::GetHealthImage() const
+{
+	switch (GetDeviceHealth())
+	{
+		case EDeviceHealth::Good: return FLiveLinkDeviceStyle::Get()->GetBrush("LiveLinkHub.Devices.Status.Good");
+		case EDeviceHealth::Info: return FLiveLinkDeviceStyle::Get()->GetBrush("LiveLinkHub.Devices.Status.Info");
+		case EDeviceHealth::Warning: return FLiveLinkDeviceStyle::Get()->GetBrush("LiveLinkHub.Devices.Status.Warning");
+		case EDeviceHealth::Error: return FLiveLinkDeviceStyle::Get()->GetBrush("LiveLinkHub.Devices.Status.Error");
+		default: return nullptr;
+	}
+}
 
 void ULiveLinkDevice::InternalDeviceAdded(const FGuid InDeviceGuid, ULiveLinkDeviceSettings* InSettings)
 {
diff --git a/Engine/Plugins/Animation/LiveLinkDevice/Source/LiveLinkDevice/Private/LiveLinkDeviceCapability_Connection.cpp b/Engine/Plugins/Animation/LiveLinkDevice/Source/LiveLinkDevice/Private/LiveLinkDeviceCapability_Connection.cpp
--- a/Engine/Plugins/Animation/LiveLinkDevice/Source/LiveLinkDevice/Private/LiveLinkDeviceCapability_Connection.cpp
+++ b/Engine/Plugins/Animation/LiveLinkDevice/Source/LiveLinkDevice/Private/LiveLinkDeviceCapability_Connection.cpp
@@ -3,6 +3,7 @@
 #include "LiveLinkDeviceCapability_Connection.h"
 #include "Engine/Engine.h"
 #include "Framework/MultiBox/MultiBoxBuilder.h"
+#include "LiveLinkDefaultColumns.h"
 #include "LiveLinkDevice.h"
 #include "LiveLinkDeviceSubsystem.h"
 #include "Styling/StyleColors.h"
@@ -47,6 +48,12 @@ protected:
 };
 
 
+ULiveLinkDeviceCapability_Connection::ULiveLinkDeviceCapability_Connection()
+	: Column_HardwareId(RegisterTableColumn("HardwareId"))
+{
+	RegisterGlobalTableColumnHandler(UE::LiveLink::DefaultColumn::Machine);
+}
+
 SHeaderRow::FColumn::FArguments& ULiveLinkDeviceCapability_Connection::GenerateHeaderForColumn(
 	const FName InColumnId,
 	SHeaderRow::FColumn::FArguments& InArgs
@@ -115,7 +122,6 @@ TSharedPtr<SWidget> ULiveLinkDeviceCapability_Connection::GenerateWidgetForColum
 		);
 }
 
-
 TSharedRef<SWidget> ULiveLinkDeviceCapability_Connection::Header_GetMenuContent()
 {
 	const bool bShouldCloseAfterSelection_true = true;
diff --git a/Engine/Plugins/Animation/LiveLinkDevice/Source/LiveLinkDevice/Private/LiveLinkDeviceCapability_Recording.cpp b/Engine/Plugins/Animation/LiveLinkDevice/Source/LiveLinkDevice/Private/LiveLinkDeviceCapability_Recording.cpp
--- a/Engine/Plugins/Animation/LiveLinkDevice/Source/LiveLinkDevice/Private/LiveLinkDeviceCapability_Recording.cpp
+++ b/Engine/Plugins/Animation/LiveLinkDevice/Source/LiveLinkDevice/Private/LiveLinkDeviceCapability_Recording.cpp
@@ -2,12 +2,18 @@
 
 #include "LiveLinkDeviceCapability_Recording.h"
 #include "Engine/Engine.h"
+#include "LiveLinkDefaultColumns.h"
 #include "LiveLinkDevice.h"
 #include "LiveLinkDeviceStyle.h"
 #include "LiveLinkDeviceSubsystem.h"
 #include "LiveLinkHub/ILiveLinkRecordingSession.h"
 #include "Widgets/Images/SImage.h"
 
+ULiveLinkDeviceCapability_Recording::ULiveLinkDeviceCapability_Recording()
+	: Column_RecordingStatus(RegisterTableColumn("RecordingStatus"))
+{
+	RegisterGlobalTableColumnHandler(UE::LiveLink::DefaultColumn::Action);
+}
 
 class SLiveLinkCapability_Recording : public SCompoundWidget
 {
@@ -111,7 +117,6 @@ TSharedPtr<SWidget> ULiveLinkDeviceCapability_Recording::GenerateWidgetForColumn
 		);
 }
 
-
 void ULiveLinkDeviceCapability_Recording::HandleRecordingStarted()
 {
 	ULiveLinkDeviceSubsystem* Subsystem = GEngine->GetEngineSubsystem<ULiveLinkDeviceSubsystem>();
diff --git a/Engine/Plugins/Animation/LiveLinkDevice/Source/LiveLinkDevice/Private/LiveLinkDeviceModule.cpp b/Engine/Plugins/Animation/LiveLinkDevice/Source/LiveLinkDevice/Private/LiveLinkDeviceModule.cpp
--- a/Engine/Plugins/Animation/LiveLinkDevice/Source/LiveLinkDevice/Private/LiveLinkDeviceModule.cpp
+++ b/Engine/Plugins/Animation/LiveLinkDevice/Source/LiveLinkDevice/Private/LiveLinkDeviceModule.cpp
@@ -81,15 +81,13 @@ void FLiveLinkDeviceModule::DeviceSelectionChanged(ULiveLinkDevice* InSelectedDe
 }
 
 
-TSharedRef<SWidget> FLiveLinkDeviceModule::OnGenerateAddDeviceMenu()
+void FLiveLinkDeviceModule::CreateDeviceMenuEntries(FMenuBuilder& MenuBuilder)
 {
 	ULiveLinkDeviceSubsystem* DeviceSubsystem = GEngine->GetEngineSubsystem<ULiveLinkDeviceSubsystem>();
 	check(DeviceSubsystem);
 
 	const TSet<TSubclassOf<ULiveLinkDevice>>& DeviceClasses = DeviceSubsystem->GetKnownDeviceClasses();
 
-	const bool bCloseAfterSelection_true = true;
-	FMenuBuilder MenuBuilder(bCloseAfterSelection_true, nullptr);
 	MenuBuilder.BeginSection("DevicesSection", LOCTEXT("DevicesSectionHeading", "Live Link Devices"));
 
 	for (const TSubclassOf<ULiveLinkDevice>& DeviceClass : DeviceClasses)
@@ -121,6 +119,15 @@ TSharedRef<SWidget> FLiveLinkDeviceModule::OnGenerateAddDeviceMenu()
 	}
 
 	MenuBuilder.EndSection();
+}
+
+TSharedRef<SWidget> FLiveLinkDeviceModule::OnGenerateAddDeviceMenu()
+{
+	const bool bCloseAfterSelection_true = true;
+	FMenuBuilder MenuBuilder(bCloseAfterSelection_true, nullptr);
+
+	CreateDeviceMenuEntries(MenuBuilder);
+
 	return MenuBuilder.MakeWidget();
 }
 
diff --git a/Engine/Plugins/Animation/LiveLinkDevice/Source/LiveLinkDevice/Private/LiveLinkDeviceModule.h b/Engine/Plugins/Animation/LiveLinkDevice/Source/LiveLinkDevice/Private/LiveLinkDeviceModule.h
--- a/Engine/Plugins/Animation/LiveLinkDevice/Source/LiveLinkDevice/Private/LiveLinkDeviceModule.h
+++ b/Engine/Plugins/Animation/LiveLinkDevice/Source/LiveLinkDevice/Private/LiveLinkDeviceModule.h
@@ -8,13 +8,13 @@
 #include "UObject/WeakObjectPtr.h"
 
 
+class FMenuBuilder;
 class FSpawnTabArgs;
 class IDetailsView;
 class SDockTab;
 class SLiveLinkDeviceTable;
 class SWidget;
 
-
 class FLiveLinkDeviceModule : public ILiveLinkDeviceModule
 {
 public:
@@ -28,6 +28,10 @@ public:
 
 	FOnDeviceSelectionChangedDelegate& OnSelectionChanged() override { return OnDeviceSelectionChangedDelegate; }
 
+	//~ Begin ILiveLinkDeviceModule interface
+	virtual void CreateDeviceMenuEntries(FMenuBuilder& Builder) override;
+	//~ End ILiveLinkDeviceModule interface
+
 private:
 	TSharedRef<SWidget> OnGenerateAddDeviceMenu();
 	TSharedRef<SDockTab> OnSpawnDevicesTab(const FSpawnTabArgs& InSpawnTabArgs);
diff --git a/Engine/Plugins/Animation/LiveLinkDevice/Source/LiveLinkDevice/Public/ILiveLinkDeviceModule.h b/Engine/Plugins/Animation/LiveLinkDevice/Source/LiveLinkDevice/Public/ILiveLinkDeviceModule.h
--- a/Engine/Plugins/Animation/LiveLinkDevice/Source/LiveLinkDevice/Public/ILiveLinkDeviceModule.h
+++ b/Engine/Plugins/Animation/LiveLinkDevice/Source/LiveLinkDevice/Public/ILiveLinkDeviceModule.h
@@ -11,6 +11,8 @@ DECLARE_MULTICAST_DELEGATE_OneParam(FOnDeviceSelectionChangedDelegate, ULiveLink
 
 DECLARE_LOG_CATEGORY_EXTERN(LogLiveLinkDevice, Log, All);
 
+class FMenuBuilder;
+
 /**
 * The public interface to this module
 */
@@ -36,4 +38,7 @@ public:
 	* @return Returns the device selection changed delegate
 	*/
 	virtual FOnDeviceSelectionChangedDelegate& OnSelectionChanged() = 0;
+
+	/** Populate a menu builder with entries for every device registered with the device subsystem. */
+	virtual void CreateDeviceMenuEntries(FMenuBuilder& Builder) = 0;
 };
diff --git a/Engine/Plugins/Animation/LiveLinkDevice/Source/LiveLinkDevice/Public/LiveLinkDevice.h b/Engine/Plugins/Animation/LiveLinkDevice/Source/LiveLinkDevice/Public/LiveLinkDevice.h
--- a/Engine/Plugins/Animation/LiveLinkDevice/Source/LiveLinkDevice/Public/LiveLinkDevice.h
+++ b/Engine/Plugins/Animation/LiveLinkDevice/Source/LiveLinkDevice/Public/LiveLinkDevice.h
@@ -87,6 +87,8 @@ public:
 private:
 	friend ULiveLinkDeviceSubsystem;
 	virtual void InternalDeviceAdded(const FGuid InDeviceGuid, ULiveLinkDeviceSettings* InSettings);
+	/** Provides an image for the status column in the combined list view. */
+	const struct FSlateBrush* GetHealthImage() const;
 
 	FGuid DeviceGuid;
 
diff --git a/Engine/Plugins/Animation/LiveLinkDevice/Source/LiveLinkDevice/Public/LiveLinkDeviceCapability.h b/Engine/Plugins/Animation/LiveLinkDevice/Source/LiveLinkDevice/Public/LiveLinkDeviceCapability.h
--- a/Engine/Plugins/Animation/LiveLinkDevice/Source/LiveLinkDevice/Public/LiveLinkDeviceCapability.h
+++ b/Engine/Plugins/Animation/LiveLinkDevice/Source/LiveLinkDevice/Public/LiveLinkDeviceCapability.h
@@ -90,6 +90,15 @@ protected:
 		return ColumnId;
 	}
 
+	/**
+	 * Call this from your derived class to register a handler for one of the existing columns in the combined subjects/sources/devices view.
+	 * For example, this can be used by a capability to allow generating a widget for the status or action column.
+	 */
+	void RegisterGlobalTableColumnHandler(const FName InGlobalColumnName)
+	{
+		TableColumns.Emplace(InGlobalColumnName, {});
+	}
+
 	FName ExpandColumnShortName(FName InShortName) const
 	{
 		FNameBuilder NameBuilder;
diff --git a/Engine/Plugins/Animation/LiveLinkDevice/Source/LiveLinkDevice/Public/LiveLinkDeviceCapability_Connection.h b/Engine/Plugins/Animation/LiveLinkDevice/Source/LiveLinkDevice/Public/LiveLinkDeviceCapability_Connection.h
--- a/Engine/Plugins/Animation/LiveLinkDevice/Source/LiveLinkDevice/Public/LiveLinkDeviceCapability_Connection.h
+++ b/Engine/Plugins/Animation/LiveLinkDevice/Source/LiveLinkDevice/Public/LiveLinkDeviceCapability_Connection.h
@@ -26,10 +26,7 @@ class LIVELINKDEVICE_API ULiveLinkDeviceCapability_Connection : public ULiveLink
 public:
 	const FName Column_HardwareId;
 
-	ULiveLinkDeviceCapability_Connection()
-		: Column_HardwareId(RegisterTableColumn("HardwareId"))
-	{
-	}
+	ULiveLinkDeviceCapability_Connection();
 
 	//~ Begin ULiveLinkDeviceCapability interface
 	virtual SHeaderRow::FColumn::FArguments& GenerateHeaderForColumn(
@@ -42,6 +39,7 @@ public:
 		const FLiveLinkDeviceWidgetArguments& InArgs,
 		ULiveLinkDevice* InDevice
 	) override;
+
 	//~ End ULiveLinkDeviceCapability interface
 
 protected:
diff --git a/Engine/Plugins/Animation/LiveLinkDevice/Source/LiveLinkDevice/Public/LiveLinkDeviceCapability_Recording.h b/Engine/Plugins/Animation/LiveLinkDevice/Source/LiveLinkDevice/Public/LiveLinkDeviceCapability_Recording.h
--- a/Engine/Plugins/Animation/LiveLinkDevice/Source/LiveLinkDevice/Public/LiveLinkDeviceCapability_Recording.h
+++ b/Engine/Plugins/Animation/LiveLinkDevice/Source/LiveLinkDevice/Public/LiveLinkDeviceCapability_Recording.h
@@ -14,10 +14,7 @@ class LIVELINKDEVICE_API ULiveLinkDeviceCapability_Recording : public ULiveLinkD
 public:
 	const FName Column_RecordingStatus;
 
-	ULiveLinkDeviceCapability_Recording()
-		: Column_RecordingStatus(RegisterTableColumn("RecordingStatus"))
-	{
-	}
+	ULiveLinkDeviceCapability_Recording();
 
 	virtual void OnDeviceSubsystemInitialized() override;
 	virtual void OnDeviceSubsystemDeinitializing() override;
diff --git a/Engine/Source/Runtime/LiveLinkInterface/Public/LiveLinkDefaultColumns.h b/Engine/Source/Runtime/LiveLinkInterface/Public/LiveLinkDefaultColumns.h
new file mode 100644
--- /dev/null
+++ b/Engine/Source/Runtime/LiveLinkInterface/Public/LiveLinkDefaultColumns.h
@@ -0,0 +1,16 @@
+// Copyright Epic Games, Inc. All Rights Reserved.
+
+#pragma once
+
+#include "UObject/NameTypes.h"
+
+// Name of the columns that appear in the LiveLinkHub source items list.
+namespace UE::LiveLink::DefaultColumn
+{
+	inline const FName Action = "Action";
+	inline const FName Name = "Name";
+	inline const FName Machine = "Machine";
+	inline const FName TranslatedRole = "TranslatedRole";
+	inline const FName Role = "Role";
+	inline const FName Status = "Status";
+}
