LiveLinkHub - Fix device table triggering an ensure when trying to generate a header column for a global colum

#Synced-CL 48735856

diff --git a/Engine/Plugins/Animation/LiveLinkDevice/Source/LiveLinkDevice/Private/Widgets/DeviceTable.cpp b/Engine/Plugins/Animation/LiveLinkDevice/Source/LiveLinkDevice/Private/Widgets/DeviceTable.cpp
--- a/Engine/Plugins/Animation/LiveLinkDevice/Source/LiveLinkDevice/Private/Widgets/DeviceTable.cpp
+++ b/Engine/Plugins/Animation/LiveLinkDevice/Source/LiveLinkDevice/Private/Widgets/DeviceTable.cpp
@@ -131,15 +131,17 @@ void SLiveLinkDeviceTable::RegenerateColumns()
 		ULiveLinkDeviceCapability* CapabilityCDO = Capability.GetDefaultObject();
 		for (const TPair<FName, ULiveLinkDeviceCapability::FDeviceTableColumnDesc>& Column : CapabilityCDO->GetTableColumns())
 		{
-			const FName ColumnId = Column.Key;
-
-			SHeaderRow::FColumn::FArguments ColumnArgs = SHeaderRow::Column(ColumnId);
-			CapabilityCDO->GenerateHeaderForColumn(ColumnId, ColumnArgs);
+			if (!Column.Value.bGlobal)
+			{
+				const FName ColumnId = Column.Key;
 
-			// Keep the "Remove Device" column at the end.
-			const int32 InsertIdx = FMath::Max(0, HeaderRow->GetColumns().Num() - 1);
+				SHeaderRow::FColumn::FArguments ColumnArgs = SHeaderRow::Column(ColumnId);
+				CapabilityCDO->GenerateHeaderForColumn(ColumnId, ColumnArgs);
 
-			HeaderRow->InsertColumn(ColumnArgs, InsertIdx);
+				// Keep the "Remove Device" column at the end.
+				const int32 InsertIdx = FMath::Max(0, HeaderRow->GetColumns().Num() - 1);
+				HeaderRow->InsertColumn(ColumnArgs, InsertIdx);
+			}
 		}
 	}
 }
diff --git a/Engine/Plugins/Animation/LiveLinkDevice/Source/LiveLinkDevice/Public/LiveLinkDeviceCapability.h b/Engine/Plugins/Animation/LiveLinkDevice/Source/LiveLinkDevice/Public/LiveLinkDeviceCapability.h
--- a/Engine/Plugins/Animation/LiveLinkDevice/Source/LiveLinkDevice/Public/LiveLinkDeviceCapability.h
+++ b/Engine/Plugins/Animation/LiveLinkDevice/Source/LiveLinkDevice/Public/LiveLinkDeviceCapability.h
@@ -35,6 +35,9 @@ public:
 	struct FDeviceTableColumnDesc
 	{
 		// TODO: Default show/hide, sorting hint, ?
+
+		// This is a global column, no header should be generated for it since it's present by default.
+		bool bGlobal = false;
 	};
 
 public:
@@ -96,7 +99,10 @@ protected:
 	 */
 	void RegisterGlobalTableColumnHandler(const FName InGlobalColumnName)
 	{
-		TableColumns.Emplace(InGlobalColumnName, {});
+		FDeviceTableColumnDesc ColumnDesc;
+		ColumnDesc.bGlobal = true;
+
+		TableColumns.Emplace(InGlobalColumnName, ColumnDesc);
 	}
 
 	FName ExpandColumnShortName(FName InShortName) const
