Live Link Hub: Fix crash when adding non-recording devices.

All devices are asked to populate the Action column by SLiveLinkClientPanelSubjectRow::GenerateWidgetForColumn, and the Recording capability registers itself to populate the Action column, but not all devices implement the recording capability.

Matches the corresponding ImplementsInterface check in SLiveLinkDeviceRowWidget::GenerateWidgetForColumn.

#rnx
#Synced-CL 49075341

diff --git a/Engine/Plugins/Animation/LiveLinkDevice/Source/LiveLinkDevice/Private/LiveLinkDevice.cpp b/Engine/Plugins/Animation/LiveLinkDevice/Source/LiveLinkDevice/Private/LiveLinkDevice.cpp
--- a/Engine/Plugins/Animation/LiveLinkDevice/Source/LiveLinkDevice/Private/LiveLinkDevice.cpp
+++ b/Engine/Plugins/Animation/LiveLinkDevice/Source/LiveLinkDevice/Private/LiveLinkDevice.cpp
@@ -27,7 +27,7 @@ TSharedRef<SWidget> ULiveLinkDevice::GenerateWidgetForColumn(const FName InColum
 	ULiveLinkDeviceSubsystem* Subsystem = GEngine->GetEngineSubsystem<ULiveLinkDeviceSubsystem>();
 	const TMap<FName, TSubclassOf<ULiveLinkDeviceCapability>>& ColumnIdToCapability = Subsystem->GetTableColumnIdToCapability();
 	const TSubclassOf<ULiveLinkDeviceCapability>* MaybeCapabilityClass = ColumnIdToCapability.Find(InColumnId);
-	if (MaybeCapabilityClass)
+	if (MaybeCapabilityClass && GetClass()->ImplementsInterface(*MaybeCapabilityClass))
 	{
 		ULiveLinkDeviceCapability* CapabilityCDO = MaybeCapabilityClass->GetDefaultObject();
 		if (TSharedPtr<SWidget> CapabilityWidget = CapabilityCDO->GenerateWidgetForColumn(InColumnId, InArgs, this))
