Integrate Capture Manager with the Data Devices panel and deprecate the now redundant Devices panel

#Synced-CL 51340686

diff --git a/Engine/Plugins/Animation/LiveLink/Source/LiveLinkEditor/Private/LiveLinkPanelController.cpp b/Engine/Plugins/Animation/LiveLink/Source/LiveLinkEditor/Private/LiveLinkPanelController.cpp
--- a/Engine/Plugins/Animation/LiveLink/Source/LiveLinkEditor/Private/LiveLinkPanelController.cpp
+++ b/Engine/Plugins/Animation/LiveLink/Source/LiveLinkEditor/Private/LiveLinkPanelController.cpp
@@ -6,6 +6,7 @@
 #include "Engine/Engine.h"
 #include "Features/IModularFeatures.h"
 #include "IDetailsView.h"
+#include "ILiveLinkDeviceModule.h"
 #include "LiveLinkClient.h"
 #include "LiveLinkClientCommands.h"
 #include "LiveLinkClientPanelViews.h"
@@ -103,11 +104,11 @@ void FLiveLinkPanelController::OnSubjectSelectionChangedHandler(FLiveLinkSubject
 	DevicesDetailsView->SetObject(nullptr);
 	SubjectsDetailsView->SetSubjectKey(FLiveLinkSubjectKey());
 
+	// Resolve the selected device (nullptr if deselected or non-device selected)
+	ULiveLinkDevice* SelectedDevice = nullptr;
+
 	if (SubjectEntry.IsValid())
 	{
-		// Find the corresponding Source entry
-		FGuid SourceGuid = SubjectEntry->SubjectKey.Source;
-
 		if (SubjectEntry->IsSource())
 		{
 			SourcesDetailsView->SetObject(SubjectEntry->GetSettings());
@@ -121,11 +122,14 @@ void FLiveLinkPanelController::OnSubjectSelectionChangedHandler(FLiveLinkSubject
 			if (TStrongObjectPtr<ULiveLinkDevice> Device = SubjectEntry->WeakDevice.Pin())
 			{
 				DevicesDetailsView->SetObject(Device->GetDeviceSettings());
+				SelectedDevice = Device.Get();
 			}
 		}
 
 		SubjectSelectionChangedDelegate.Broadcast(SubjectEntry->SubjectKey);
 	}
+
+	ILiveLinkDeviceModule::Get().OnSelectionChanged().Broadcast(SelectedDevice);
 }
 
 void FLiveLinkPanelController::OnFinishedChangingDeviceProperties(const FPropertyChangedEvent& InPropertyChangedEvent)
diff --git a/Engine/Plugins/Animation/LiveLinkDevice/Source/LiveLinkDevice/Private/LiveLinkDeviceModule.cpp b/Engine/Plugins/Animation/LiveLinkDevice/Source/LiveLinkDevice/Private/LiveLinkDeviceModule.cpp
--- a/Engine/Plugins/Animation/LiveLinkDevice/Source/LiveLinkDevice/Private/LiveLinkDeviceModule.cpp
+++ b/Engine/Plugins/Animation/LiveLinkDevice/Source/LiveLinkDevice/Private/LiveLinkDeviceModule.cpp
@@ -9,20 +9,13 @@
 #include "LiveLinkDeviceStyle.h"
 #include "LiveLinkDeviceSubsystem.h"
 #include "Logging/StructuredLog.h"
-#include "Misc/ConfigCacheIni.h"
 #include "PropertyEditorModule.h"
 #include "SPositiveActionButton.h"
 #include "Widgets/Docking/SDockTab.h"
-#include "Widgets/DeviceTable.h"
 
 
 DEFINE_LOG_CATEGORY(LogLiveLinkDevice);
 
-
-const FName FLiveLinkDeviceModule::DevicesTabName = FName("LiveLinkDevices");
-const FName FLiveLinkDeviceModule::DeviceDetailsTabName = FName("LiveLinkDeviceDetails");
-
-
 #define LOCTEXT_NAMESPACE "LiveLinkDevice"
 
 
@@ -35,35 +28,11 @@ IMPLEMENT_MODULE(FLiveLinkDeviceModule, LiveLinkDevice);
 void FLiveLinkDeviceModule::StartupModule()
 {
 	FLiveLinkDeviceStyle::Initialize();
-
-	bIsLiveLinkHubApp = GConfig->GetBoolOrDefault(TEXT("LiveLink"), TEXT("bCreateLiveLinkHubInstance"), false, GEngineIni);
-
-	if (bIsLiveLinkHubApp)
-	{
-		TSharedRef<FGlobalTabmanager> GlobalTabManager = FGlobalTabmanager::Get();
-	
-		GlobalTabManager->RegisterNomadTabSpawner(DevicesTabName,
-				FOnSpawnTab::CreateRaw(this, &FLiveLinkDeviceModule::OnSpawnDevicesTab))
-			.SetIcon(FSlateIcon("LiveLinkDeviceStyle", "LiveLinkHub.Devices.Icon"))
-			.SetDisplayName(LOCTEXT("DevicesTabDisplayName", "Devices"));
-
-		GlobalTabManager->RegisterNomadTabSpawner(DeviceDetailsTabName,
-				FOnSpawnTab::CreateRaw(this, &FLiveLinkDeviceModule::OnSpawnDeviceDetailsTab))
-			.SetIcon(FSlateIcon(FAppStyle::GetAppStyleSetName(), "LevelEditor.Tabs.Details"))
-			.SetDisplayName(LOCTEXT("DeviceDetailsTabDisplayName", "Device Details"));
-	}
 }
 
 
 void FLiveLinkDeviceModule::ShutdownModule()
 {
-	if (bIsLiveLinkHubApp)		
-	{
-		TSharedRef<FGlobalTabmanager> GlobalTabManager = FGlobalTabmanager::Get();
-		GlobalTabManager->UnregisterNomadTabSpawner(DevicesTabName);
-		GlobalTabManager->UnregisterNomadTabSpawner(DeviceDetailsTabName);
-	}
-
 	FLiveLinkDeviceStyle::Shutdown();
 }
 
@@ -121,118 +90,4 @@ void FLiveLinkDeviceModule::CreateDeviceMenuEntries(FMenuBuilder& MenuBuilder)
 	MenuBuilder.EndSection();
 }
 
-TSharedRef<SWidget> FLiveLinkDeviceModule::OnGenerateAddDeviceMenu()
-{
-	const bool bCloseAfterSelection_true = true;
-	FMenuBuilder MenuBuilder(bCloseAfterSelection_true, nullptr);
-
-	CreateDeviceMenuEntries(MenuBuilder);
-
-	return MenuBuilder.MakeWidget();
-}
-
-
-TSharedRef<SDockTab> FLiveLinkDeviceModule::OnSpawnDevicesTab(const FSpawnTabArgs& InSpawnTabArgs)
-{
-	TSharedRef<SDockTab> DockTab = SNew(SDockTab)
-		.TabRole(ETabRole::NomadTab);
-
-	DeviceTable = SNew(SLiveLinkDeviceTable, DockTab)
-		.OnSelectionChanged_Raw(this, &FLiveLinkDeviceModule::DeviceSelectionChanged);
-
-	DockTab->SetContent(
-		SNew(SVerticalBox)
-		+ SVerticalBox::Slot()
-		.AutoHeight()
-		[
-			SNew(SBorder)
-			.BorderImage(FAppStyle::GetBrush("ToolPanel.GroupBorder"))
-			.Padding(8.0f)
-			[
-				SNew(SHorizontalBox)
-				+ SHorizontalBox::Slot()
-				.AutoWidth()
-				.Padding(0.f, 5.f, 0.f, 5.f)
-				[
-					SNew(SImage)
-					.ColorAndOpacity(FSlateColor::UseForeground())
-					.Image(FSlateIcon("LiveLinkDeviceStyle", "LiveLinkHub.Devices.Icon").GetIcon())
-				]
-				+ SHorizontalBox::Slot()
-				.HAlign(HAlign_Left)
-				.Padding(FMargin(4.0, 2.0))
-				[
-					SNew(STextBlock)
-					.Font(FCoreStyle::GetDefaultFontStyle("Regular", 14))
-					.Text(LOCTEXT("DevicesHeaderText", "Devices"))
-				]
-				+ SHorizontalBox::Slot()
-				.AutoWidth()
-				.VAlign(VAlign_Center)
-				[
-					SNew(SPositiveActionButton)
-					.Icon(FAppStyle::Get().GetBrush("Icons.Plus"))
-					.Text(LOCTEXT("AddDevice", "Add Device"))
-					.ToolTipText(LOCTEXT("AddDevice_Tooltip", "Add a new Live Link device"))
-					.OnGetMenuContent_Raw(this, &FLiveLinkDeviceModule::OnGenerateAddDeviceMenu)
-				]
-			]
-		]
-		+ SVerticalBox::Slot()
-		[
-			DeviceTable.ToSharedRef()
-		]
-	);
-
-	return DockTab;
-}
-
-
-TSharedRef<SDockTab> FLiveLinkDeviceModule::OnSpawnDeviceDetailsTab(const FSpawnTabArgs& InSpawnTabArgs)
-{
-	FPropertyEditorModule& PropertyEditorModule =
-		FModuleManager::GetModuleChecked<FPropertyEditorModule>("PropertyEditor");
-
-	FDetailsViewArgs DetailsViewArgs;
-	DetailsViewArgs.NameAreaSettings = FDetailsViewArgs::HideNameArea;
-	DetailsViewArgs.bShowPropertyMatrixButton = false;
-	DetailsViewArgs.bShowKeyablePropertiesOption = false;
-	DetailsViewArgs.bShowAnimatedPropertiesOption = false;
-
-	DetailsView = PropertyEditorModule.CreateDetailView(DetailsViewArgs);
-	if (ULiveLinkDevice* SelectedDevice = WeakSelectedDevice.Get())
-	{
-		DetailsView->SetObject(SelectedDevice->GetDeviceSettings());
-	}
-
-	DetailsView->OnFinishedChangingProperties().AddLambda(
-		[this]
-		(const FPropertyChangedEvent& InPropertyChangedEvent)
-		{
-			ULiveLinkDevice* Device = nullptr;
-
-			const TArray<TWeakObjectPtr<UObject>>& SelectedObjects = DetailsView->GetSelectedObjects();
-			if (ensure(SelectedObjects.Num() > 0))
-			{
-				if (UObject* SelectedObject = SelectedObjects[0].Get(); ensure(SelectedObject))
-				{
-					Device = SelectedObject->GetTypedOuter<ULiveLinkDevice>();
-				}
-			}
-
-			if (ensure(Device))
-			{
-				Device->OnSettingChanged(InPropertyChangedEvent);
-			}
-		}
-	);
-
-	return SNew(SDockTab)
-		.TabRole(ETabRole::NomadTab)
-		[
-			DetailsView.ToSharedRef()
-		];
-}
-
-
 #undef LOCTEXT_NAMESPACE
diff --git a/Engine/Plugins/Animation/LiveLinkDevice/Source/LiveLinkDevice/Private/LiveLinkDeviceModule.h b/Engine/Plugins/Animation/LiveLinkDevice/Source/LiveLinkDevice/Private/LiveLinkDeviceModule.h
--- a/Engine/Plugins/Animation/LiveLinkDevice/Source/LiveLinkDevice/Private/LiveLinkDeviceModule.h
+++ b/Engine/Plugins/Animation/LiveLinkDevice/Source/LiveLinkDevice/Private/LiveLinkDeviceModule.h
@@ -18,9 +18,6 @@ class SWidget;
 class FLiveLinkDeviceModule : public ILiveLinkDeviceModule
 {
 public:
-	static const FName DevicesTabName;
-	static const FName DeviceDetailsTabName;
-
 	//~ Begin IModuleInterface
 	virtual void StartupModule() override;
 	virtual void ShutdownModule() override;
@@ -32,16 +29,9 @@ public:
 	virtual void CreateDeviceMenuEntries(FMenuBuilder& Builder) override;
 	//~ End ILiveLinkDeviceModule interface
 
-private:
-	TSharedRef<SWidget> OnGenerateAddDeviceMenu();
-	TSharedRef<SDockTab> OnSpawnDevicesTab(const FSpawnTabArgs& InSpawnTabArgs);
-	TSharedRef<SDockTab> OnSpawnDeviceDetailsTab(const FSpawnTabArgs& InSpawnTabArgs);
-
 private:
 	void DeviceSelectionChanged(ULiveLinkDevice* InSelectedDevice);
 
-	bool bIsLiveLinkHubApp = false;
-	TSharedPtr<SLiveLinkDeviceTable> DeviceTable;
 	TSharedPtr<IDetailsView> DetailsView;
 	TWeakObjectPtr<ULiveLinkDevice> WeakSelectedDevice;
 
diff --git a/Engine/Plugins/Animation/LiveLinkDevice/Source/LiveLinkDevice/Private/Widgets/DeviceTable.cpp b/Engine/Plugins/Animation/LiveLinkDevice/Source/LiveLinkDevice/Private/Widgets/DeviceTable.cpp
deleted file mode 100644
diff --git a/Engine/Plugins/Animation/LiveLinkDevice/Source/LiveLinkDevice/Private/Widgets/DeviceTable.h b/Engine/Plugins/Animation/LiveLinkDevice/Source/LiveLinkDevice/Private/Widgets/DeviceTable.h
deleted file mode 100644
diff --git a/Engine/Plugins/VirtualProduction/CaptureManager/CaptureManagerApp/Source/CaptureManagerEditor/CaptureManagerEditor.Build.cs b/Engine/Plugins/VirtualProduction/CaptureManager/CaptureManagerApp/Source/CaptureManagerEditor/CaptureManagerEditor.Build.cs
--- a/Engine/Plugins/VirtualProduction/CaptureManager/CaptureManagerApp/Source/CaptureManagerEditor/CaptureManagerEditor.Build.cs
+++ b/Engine/Plugins/VirtualProduction/CaptureManager/CaptureManagerApp/Source/CaptureManagerEditor/CaptureManagerEditor.Build.cs
@@ -18,6 +18,7 @@ public class CaptureManagerEditor : ModuleRules
 		{
 			"InputCore",
 			"LiveLinkDevice",
+			"LiveLinkEditor",
 			"LiveLinkCapabilities",
 			"LiveLinkHub",
 			"Slate",
diff --git a/Engine/Plugins/VirtualProduction/CaptureManager/CaptureManagerApp/Source/CaptureManagerEditor/Private/UI/CaptureManagerPanelController.cpp b/Engine/Plugins/VirtualProduction/CaptureManager/CaptureManagerApp/Source/CaptureManagerEditor/Private/UI/CaptureManagerPanelController.cpp
--- a/Engine/Plugins/VirtualProduction/CaptureManager/CaptureManagerApp/Source/CaptureManagerEditor/Private/UI/CaptureManagerPanelController.cpp
+++ b/Engine/Plugins/VirtualProduction/CaptureManager/CaptureManagerApp/Source/CaptureManagerEditor/Private/UI/CaptureManagerPanelController.cpp
@@ -14,6 +14,7 @@
 #include "Settings/CaptureManagerSettings.h"
 
 #include "Engine/Engine.h"
+#include "ILiveLinkDeviceModule.h"
 #include "Ingest/LiveLinkDeviceCapability_Ingest.h"
 #include "LiveLinkDeviceCapability_Connection.h"
 
@@ -66,6 +67,10 @@ FCaptureManagerPanelController::FCaptureManagerPanelController(FPrivateToken InP
 
 FCaptureManagerPanelController::~FCaptureManagerPanelController()
 {
+	if (FModuleManager::Get().IsModuleLoaded("LiveLinkDevice"))
+	{
+		ILiveLinkDeviceModule::Get().OnSelectionChanged().Remove(DeviceSelectionChangedHandle);
+	}
 	IngestJobDetailsView->OnFinishedChangingProperties().Remove(FinishedChangingPropertiesHandle);
 
 	IngestJobProcessorWidget->OnJobsAdded().Unbind();
@@ -83,6 +88,8 @@ void FCaptureManagerPanelController::CreateViews()
 	CreateIngestManagementViews();
 
 	TakesView = MakeShared<FTakesView>(AsShared());
+
+	DeviceSelectionChangedHandle = ILiveLinkDeviceModule::Get().OnSelectionChanged().AddSP(this, &FCaptureManagerPanelController::OnDeviceSelectionChanged);
 }
 
 void FCaptureManagerPanelController::CreateIngestManagementViews()
@@ -201,6 +208,11 @@ TSharedRef<FUpdateTakeListManager> FCaptureManagerPanelController::GetUpdateTake
 	return UpdateTakeListManager.ToSharedRef();
 }
 
+void FCaptureManagerPanelController::OnDeviceSelectionChanged(ULiveLinkDevice* InDevice)
+{
+	TakesView->DeviceSelectionChanged(InDevice);
+}
+
 void FCaptureManagerPanelController::OnDeviceAdded(FGuid InDeviceId, ULiveLinkDevice* InDevice)
 {
 	if (!InDevice->Implements<ULiveLinkDeviceCapability_Ingest>())
diff --git a/Engine/Plugins/VirtualProduction/CaptureManager/CaptureManagerApp/Source/CaptureManagerEditor/Private/UI/CaptureManagerPanelController.h b/Engine/Plugins/VirtualProduction/CaptureManager/CaptureManagerApp/Source/CaptureManagerEditor/Private/UI/CaptureManagerPanelController.h
--- a/Engine/Plugins/VirtualProduction/CaptureManager/CaptureManagerApp/Source/CaptureManagerEditor/Private/UI/CaptureManagerPanelController.h
+++ b/Engine/Plugins/VirtualProduction/CaptureManager/CaptureManagerApp/Source/CaptureManagerEditor/Private/UI/CaptureManagerPanelController.h
@@ -61,6 +61,7 @@ public:
 
 private:
 
+	void OnDeviceSelectionChanged(ULiveLinkDevice* InDevice);
 	void OnDeviceAdded(FGuid InDeviceId, ULiveLinkDevice* InDevice);
 	void OnDeviceRemoved(FGuid InDeviceId, ULiveLinkDevice* InDevice);
 	void OnReachableEvent(ELiveLinkDeviceConnectionStatus InStatus, FGuid InDeviceId);
@@ -86,6 +87,7 @@ private:
 	TSharedPtr<UE::CaptureManager::SIngestJobProcessor> IngestJobProcessorWidget;
 	TSharedPtr<IDetailsView> IngestJobDetailsView;
 	FDelegateHandle FinishedChangingPropertiesHandle;
+	FDelegateHandle DeviceSelectionChangedHandle;
 
 	TSharedPtr<FUpdateTakeListManager> UpdateTakeListManager;
 };
\ No newline at end of file
diff --git a/Engine/Plugins/VirtualProduction/CaptureManager/CaptureManagerApp/Source/CaptureManagerEditor/Private/UI/LiveLinkHubCaptureManagerMode.cpp b/Engine/Plugins/VirtualProduction/CaptureManager/CaptureManagerApp/Source/CaptureManagerEditor/Private/UI/LiveLinkHubCaptureManagerMode.cpp
--- a/Engine/Plugins/VirtualProduction/CaptureManager/CaptureManagerApp/Source/CaptureManagerEditor/Private/UI/LiveLinkHubCaptureManagerMode.cpp
+++ b/Engine/Plugins/VirtualProduction/CaptureManager/CaptureManagerApp/Source/CaptureManagerEditor/Private/UI/LiveLinkHubCaptureManagerMode.cpp
@@ -6,6 +6,8 @@
 #include "CaptureManagerPanelController.h"
 #include "CaptureUtilsModule.h"
 #include "LiveLinkHubApplicationBase.h"
+#include "LiveLinkClientPanelViews.h"
+#include "LiveLinkPanelController.h"
 #include "Network/NetworkMisc.h"
 #include "SIngestJobProcessor.h"
 #include "TakesView.h"
@@ -78,7 +80,7 @@ struct FJobsDetailsTabSummoner : public FWorkflowTabFactory
 public:
 	FJobsDetailsTabSummoner(TSharedPtr<FLiveLinkHubApplicationBase> InHostingApp, TSharedPtr<FCaptureManagerPanelController> InPanelController)
 		: FWorkflowTabFactory(UE::CaptureManager::JobDetailsTabId, MoveTemp(InHostingApp))
-		, PanelController(MoveTemp(InPanelController))
+		, CaptureManagerPanelController(MoveTemp(InPanelController))
 	{
 		TabLabel = LOCTEXT("JobDetailsTabLabel", "Job Details");
 		TabIcon = FSlateIcon(FAppStyle::GetAppStyleSetName(), "LevelEditor.Tabs.Details");
@@ -90,20 +92,21 @@ public:
 
 	virtual TSharedRef<SWidget> CreateTabBody(const FWorkflowTabSpawnInfo& Info) const override
 	{
-		return PanelController->GetIngestJobDetailsWidget();
+		return CaptureManagerPanelController->GetIngestJobDetailsWidget();
 	}
 
 private:
-	TSharedPtr<FCaptureManagerPanelController> PanelController;
+	TSharedPtr<FCaptureManagerPanelController> CaptureManagerPanelController;
 };
 
+
 FLiveLinkHubCaptureManagerMode::FLiveLinkHubCaptureManagerMode(TSharedPtr<FLiveLinkHubApplicationBase> App)
 	: FLiveLinkHubApplicationMode("CaptureManager", LOCTEXT("CaptureManagerModeLabel", "Capture Manager"), App),
 	UnrealEndpointManager(FModuleManager::LoadModuleChecked<FCaptureManagerUnrealEndpointModule>("CaptureManagerUnrealEndpoint").GetEndpointManager())
 {
 	using namespace UE::CaptureManager;
 
-	TabLayout = FTabManager::NewLayout("LiveLinkCaptureManagerMode_v1.2");
+	TabLayout = FTabManager::NewLayout("LiveLinkCaptureManagerMode_v1.3");
 	
 	TabLayout->AddArea
 		(
@@ -123,14 +126,13 @@ FLiveLinkHubCaptureManagerMode::FLiveLinkHubCaptureManagerMode(TSharedPtr<FLiveL
 					(
 						FTabManager::NewStack()
 						->SetSizeCoefficient(0.5f)
-						->SetHideTabWell(true)
-						->AddTab("LiveLinkDevices", ETabState::OpenedTab)
+						->AddTab(UE::LiveLinkHub::SubjectsTabId, ETabState::OpenedTab)
 					)
 					->Split
 					(
 						FTabManager::NewStack()
 						->SetSizeCoefficient(0.5f)
-						->AddTab("LiveLinkDeviceDetails", ETabState::OpenedTab)
+						->AddTab(UE::LiveLinkHub::SourceDetailsTabId, ETabState::OpenedTab)
 					)
 				)
 				->Split
@@ -166,12 +168,11 @@ FLiveLinkHubCaptureManagerMode::FLiveLinkHubCaptureManagerMode(TSharedPtr<FLiveL
 			)
 		);
 	
-	
-	PanelController = FCaptureManagerPanelController::MakeInstance();
+	CaptureManagerPanelController = FCaptureManagerPanelController::MakeInstance();
 
-	TabFactories.RegisterFactory(MakeShared<FTakeBrowserTabSummoner>(App, PanelController));
-	TabFactories.RegisterFactory(MakeShared<FJobsListTabSummoner>(App, PanelController));
-	TabFactories.RegisterFactory(MakeShared<FJobsDetailsTabSummoner>(App, PanelController));
+	TabFactories.RegisterFactory(MakeShared<FTakeBrowserTabSummoner>(App, CaptureManagerPanelController));
+	TabFactories.RegisterFactory(MakeShared<FJobsListTabSummoner>(App, CaptureManagerPanelController));
+	TabFactories.RegisterFactory(MakeShared<FJobsDetailsTabSummoner>(App, CaptureManagerPanelController));
 }
 
 FLiveLinkHubCaptureManagerMode::~FLiveLinkHubCaptureManagerMode() = default;
diff --git a/Engine/Plugins/VirtualProduction/CaptureManager/CaptureManagerApp/Source/CaptureManagerEditor/Private/UI/LiveLinkHubCaptureManagerMode.h b/Engine/Plugins/VirtualProduction/CaptureManager/CaptureManagerApp/Source/CaptureManagerEditor/Private/UI/LiveLinkHubCaptureManagerMode.h
--- a/Engine/Plugins/VirtualProduction/CaptureManager/CaptureManagerApp/Source/CaptureManagerEditor/Private/UI/LiveLinkHubCaptureManagerMode.h
+++ b/Engine/Plugins/VirtualProduction/CaptureManager/CaptureManagerApp/Source/CaptureManagerEditor/Private/UI/LiveLinkHubCaptureManagerMode.h
@@ -18,7 +18,6 @@
 
 
 class FCommunicationManager;
-
 class SCaptureManagerMainTabView;
 class SDockTab;
 
@@ -41,6 +40,6 @@ private:
 
 private:
 	TSharedRef<UE::CaptureManager::FUnrealEndpointManager> UnrealEndpointManager;
-	TSharedPtr<class FCaptureManagerPanelController> PanelController;
+	TSharedPtr<class FCaptureManagerPanelController> CaptureManagerPanelController;
 };
 
diff --git a/Engine/Plugins/VirtualProduction/CaptureManager/CaptureManagerApp/Source/CaptureManagerEditor/Private/UI/TakesView.cpp b/Engine/Plugins/VirtualProduction/CaptureManager/CaptureManagerApp/Source/CaptureManagerEditor/Private/UI/TakesView.cpp
--- a/Engine/Plugins/VirtualProduction/CaptureManager/CaptureManagerApp/Source/CaptureManagerEditor/Private/UI/TakesView.cpp
+++ b/Engine/Plugins/VirtualProduction/CaptureManager/CaptureManagerApp/Source/CaptureManagerEditor/Private/UI/TakesView.cpp
@@ -579,13 +579,6 @@ void FTakesView::CreateTakesView()
 	TSharedRef<UE::CaptureManager::FIngestPipelineManager> IngestPipelineManager = Controller->GetIngestPipelineManager();
 	TSharedRef<UE::CaptureManager::FIngestJobSettingsManager> IngestJobSettingsManager = Controller->GetIngestJobSettingsManager();
 
-	ILiveLinkDeviceModule& Module = ILiveLinkDeviceModule::Get();
-	Module.OnSelectionChanged().AddLambda([this](ULiveLinkDevice* InSelectedDevice)
-		{
-			SelectedItem = InSelectedDevice;
-			RefreshTakesViewDelegate.ExecuteIfBound(true);
-		});
-
 	TUniquePtr<FGetCurrentSelectionDelegate> GetCurrentSelectionDelegate = MakeUnique<FGetCurrentSelectionDelegate>();
 
 	FAssetPickerConfig TakesPickerConfig;
@@ -835,6 +828,12 @@ void FTakesView::CaptureDeviceRemoved(ULiveLinkDevice* InDevice)
 		});
 }
 
+void FTakesView::DeviceSelectionChanged(ULiveLinkDevice* InDevice)
+{
+	SelectedItem = InDevice;
+	RefreshTakesViewDelegate.ExecuteIfBound(true);
+}
+
 TArray<TUniqueObjectPtr<UTakeVirtualAsset>> FTakesView::MakeTakeObjects(FGuid InCaptureDeviceId, TOptional<TArray<UE::CaptureManager::FTakeId>> InTakeIds)
 {
 	TArray<TUniqueObjectPtr<UTakeVirtualAsset>> TakeObjects;
diff --git a/Engine/Plugins/VirtualProduction/CaptureManager/CaptureManagerApp/Source/CaptureManagerEditor/Private/UI/TakesView.h b/Engine/Plugins/VirtualProduction/CaptureManager/CaptureManagerApp/Source/CaptureManagerEditor/Private/UI/TakesView.h
--- a/Engine/Plugins/VirtualProduction/CaptureManager/CaptureManagerApp/Source/CaptureManagerEditor/Private/UI/TakesView.h
+++ b/Engine/Plugins/VirtualProduction/CaptureManager/CaptureManagerApp/Source/CaptureManagerEditor/Private/UI/TakesView.h
@@ -86,6 +86,7 @@ public:
 	void CaptureDeviceStopped(FGuid InCaptureDeviceId);
 	void CaptureDeviceAdded(ULiveLinkDevice* InDevice);
 	void CaptureDeviceRemoved(ULiveLinkDevice* InDevice);
+	void DeviceSelectionChanged(ULiveLinkDevice* InDevice);
 
 	TSharedPtr<STakesView> TakesTileView;
 private:
