Adding interfaces from super classes as well

#Synced-CL 51454069

diff --git a/Engine/Plugins/Animation/LiveLinkDevice/Source/LiveLinkDevice/Private/LiveLinkDeviceSubsystem.cpp b/Engine/Plugins/Animation/LiveLinkDevice/Source/LiveLinkDevice/Private/LiveLinkDeviceSubsystem.cpp
--- a/Engine/Plugins/Animation/LiveLinkDevice/Source/LiveLinkDevice/Private/LiveLinkDeviceSubsystem.cpp
+++ b/Engine/Plugins/Animation/LiveLinkDevice/Source/LiveLinkDevice/Private/LiveLinkDeviceSubsystem.cpp
@@ -356,20 +356,23 @@ void ULiveLinkDeviceSubsystem::RegisterDeviceClass(TSubclassOf<ULiveLinkDevice>
 
 	KnownDeviceClasses.Add(InDeviceClass);
 
-	for (TArray<FImplementedInterface>::TConstIterator It(InDeviceClass->Interfaces); It; ++It)
+	for (UClass* TempClass = InDeviceClass; TempClass; TempClass = TempClass->GetSuperClass())
 	{
-		UClass* InterfaceClass = It->Class;
-		if (!InterfaceClass || !InterfaceClass->IsChildOf(ULiveLinkDeviceCapability::StaticClass()))
+		for (TArray<FImplementedInterface>::TConstIterator It(TempClass->Interfaces); It; ++It)
 		{
-			continue;
-		}
+			UClass* InterfaceClass = It->Class;
+			if (!InterfaceClass || !InterfaceClass->IsChildOf(ULiveLinkDeviceCapability::StaticClass()))
+			{
+				continue;
+			}
 
-		TSubclassOf<ULiveLinkDeviceCapability> Capability = InterfaceClass;
+			TSubclassOf<ULiveLinkDeviceCapability> Capability = InterfaceClass;
 
-		// All capabilities should have been enumerated prior to device classes.
-		ensure(KnownCapabilities.Find(Capability));
+			// All capabilities should have been enumerated prior to device classes.
+			ensure(KnownCapabilities.Find(Capability));
 
-		CapabilitiesByDeviceClass.AddUnique(InDeviceClass, Capability);
-		DeviceClassesByCapability.AddUnique(Capability, InDeviceClass);
+			CapabilitiesByDeviceClass.AddUnique(InDeviceClass, Capability);
+			DeviceClassesByCapability.AddUnique(Capability, InDeviceClass);
+		}
 	}
 }
